Temporarily disabling ARKit until the build machines can update their XCode builds.

diff --git a/Engine/Plugins/Runtime/AppleARKit/AppleARKit.uplugin b/Engine/Plugins/Runtime/AppleARKit/AppleARKit.uplugin
--- a/Engine/Plugins/Runtime/AppleARKit/AppleARKit.uplugin
+++ b/Engine/Plugins/Runtime/AppleARKit/AppleARKit.uplugin
@@ -17,7 +17,7 @@
 			"Name": "AppleARKit",
 			"Type": "Runtime",
 			"LoadingPhase": "PostConfigInit",
-			"WhitelistPlatforms" : [ "IOS", "Mac", "Win64" ]
+			"WhitelistPlatforms" : [ "Mac", "Win64" ]
 		},
 		{
 			"Name": "AppleARKitEditor",
