Removing Mac as a supported ARKit platform temporarily

diff --git a/Engine/Plugins/Runtime/AppleARKit/AppleARKit.uplugin b/Engine/Plugins/Runtime/AppleARKit/AppleARKit.uplugin
--- a/Engine/Plugins/Runtime/AppleARKit/AppleARKit.uplugin
+++ b/Engine/Plugins/Runtime/AppleARKit/AppleARKit.uplugin
@@ -17,13 +17,13 @@
 			"Name": "AppleARKit",
 			"Type": "Runtime",
 			"LoadingPhase": "PostConfigInit",
-			"WhitelistPlatforms" : [ "Mac", "Win64" ]
+			"WhitelistPlatforms" : [ "Win64" ]
 		},
 		{
 			"Name": "AppleARKitEditor",
 			"Type": "Editor",
 			"LoadingPhase": "Default",
-			"WhitelistPlatforms" : [ "Mac", "Win64" ]
+			"WhitelistPlatforms" : [ "Win64" ]
 		}
 	]
 }
