Proper fix for ARKit interfering with automated tests.

diff --git a/Engine/Plugins/Runtime/AppleARKit/Source/AppleARKit/Private/AppleARKitSystem.cpp b/Engine/Plugins/Runtime/AppleARKit/Source/AppleARKit/Private/AppleARKitSystem.cpp
--- a/Engine/Plugins/Runtime/AppleARKit/Source/AppleARKit/Private/AppleARKitSystem.cpp
+++ b/Engine/Plugins/Runtime/AppleARKit/Source/AppleARKit/Private/AppleARKitSystem.cpp
@@ -62,7 +62,11 @@ private:
 	
 	virtual bool IsActiveThisFrame(class FViewport* InViewport) const override
 	{
+#if ARKIT_SUPPORT
 		return true;
+#else
+		return false;
+#endif //ARKIT_SUPPORT
 	}
 	//~ FDefaultXRCamera
 	
