CIS Fix

diff --git a/Engine/Plugins/Runtime/AppleARKit/Source/AppleARKit/Private/AppleARKitHitTestResult.h b/Engine/Plugins/Runtime/AppleARKit/Source/AppleARKit/Private/AppleARKitHitTestResult.h
--- a/Engine/Plugins/Runtime/AppleARKit/Source/AppleARKit/Private/AppleARKitHitTestResult.h
+++ b/Engine/Plugins/Runtime/AppleARKit/Source/AppleARKit/Private/AppleARKitHitTestResult.h
@@ -51,7 +51,7 @@ struct FAppleARKitHitTestResult
     // Default constructor
 	FAppleARKitHitTestResult() {};
 
-#if ARKIT_SUPPORT
+#if ARKIT_SUPPORT && __IPHONE_OS_VERSION_MAX_ALLOWED >= 110000
 
 	/** 
 	 * This is a conversion copy-constructor that takes a raw ARHitTestResult and fills this 
