Merging AR environment texture fix using Dev-VR_to_Release-4.20

diff --git a/Engine/Plugins/Runtime/AppleARKit/Source/AppleARKit/Private/AppleARKitTextures.cpp b/Engine/Plugins/Runtime/AppleARKit/Source/AppleARKit/Private/AppleARKitTextures.cpp
--- a/Engine/Plugins/Runtime/AppleARKit/Source/AppleARKit/Private/AppleARKitTextures.cpp
+++ b/Engine/Plugins/Runtime/AppleARKit/Source/AppleARKit/Private/AppleARKitTextures.cpp
@@ -230,26 +230,25 @@ public:
 		if (MetalTexture != nullptr)
 		{
 			Size.X = Size.Y = Owner->Size.X;
-			const int32 NumComponents = 4;
-			const int32 BytesPerImage = Size.X * Size.Y * NumComponents;
-			const int32 BytesPerRow = Size.X * NumComponents;
 
 			const uint32 CreateFlags = TexCreate_NoMipTail | TexCreate_SRGB;
 			EnvCubemapTextureRHIRef = RHICreateTextureCube(Size.X, PF_B8G8R8A8, 1, CreateFlags, CreateInfo);
-			/** To map their texture faces into our space we need:
-				-Z to +X CCW 90
-				+Z to -X CW 90
-				+X to +Y CCW 90
-				-X to -Y CCW 90
-				+Y to +Z CW 180
-				-Y to -Z
+
+			/**
+			 * To map their texture faces into our space we need:
+			 *	 +X	to +Y	Down Mirrored
+			 *	 -X to -Y	Up Mirrored
+			 *	 +Y to +Z	Left Mirrored
+			 *	 -Y to -Z	Left Mirrored
+			 *	 +Z to -X	Left Mirrored
+			 *	 -Z to +X	Right Mirrored
 			 */
-			CopyCubeFace(MetalTexture, EnvCubemapTextureRHIRef, kCGImagePropertyOrientationLeft,  5, 0);
-			CopyCubeFace(MetalTexture, EnvCubemapTextureRHIRef, kCGImagePropertyOrientationRight, 4, 1);
-			CopyCubeFace(MetalTexture, EnvCubemapTextureRHIRef, kCGImagePropertyOrientationLeft, 0, 2);
-			CopyCubeFace(MetalTexture, EnvCubemapTextureRHIRef, kCGImagePropertyOrientationLeft, 1, 3);
-			CopyCubeFace(MetalTexture, EnvCubemapTextureRHIRef, kCGImagePropertyOrientationDown, 2, 4);
-			CopyCubeFace(MetalTexture, EnvCubemapTextureRHIRef, kCGImagePropertyOrientationUp, 3, 5);
+			CopyCubeFace(MetalTexture, EnvCubemapTextureRHIRef, kCGImagePropertyOrientationDownMirrored, 0, 2);
+			CopyCubeFace(MetalTexture, EnvCubemapTextureRHIRef, kCGImagePropertyOrientationUpMirrored, 1, 3);
+			CopyCubeFace(MetalTexture, EnvCubemapTextureRHIRef, kCGImagePropertyOrientationLeftMirrored, 2, 4);
+			CopyCubeFace(MetalTexture, EnvCubemapTextureRHIRef, kCGImagePropertyOrientationLeftMirrored, 3, 5);
+			CopyCubeFace(MetalTexture, EnvCubemapTextureRHIRef, kCGImagePropertyOrientationLeftMirrored, 4, 1);
+			CopyCubeFace(MetalTexture, EnvCubemapTextureRHIRef, kCGImagePropertyOrientationRightMirrored, 5, 0);
 		}
 		else
 		{
@@ -258,6 +257,7 @@ public:
 			EnvCubemapTextureRHIRef = RHICreateTextureCube(Size.X, PF_B8G8R8A8, 1, 0, CreateInfo);
 		}
 
+
 		TextureRHI = EnvCubemapTextureRHIRef;
 		TextureRHI->SetName(Owner->GetFName());
 		RHIBindDebugLabelName(TextureRHI, *Owner->GetName());
@@ -267,7 +267,6 @@ public:
 		SamplerStateRHI = RHICreateSamplerState(SamplerStateInitializer);
 	}
 
-//@joeg -- Move to image utils plugin
 	void CopyCubeFace(id<MTLTexture> MetalTexture, FTextureCubeRHIRef Cubemap, uint32 Rotation, int32 MetalCubeIndex, int32 OurCubeIndex)
 	{
 		// Rotate the image we need to get a view into the face as a new slice
@@ -286,15 +285,17 @@ public:
 		}
 
 		// Make a new view into our texture and directly render to that to avoid the CPU copy
-		id<MTLTexture> OurCubeFaceMetalTexture = [(id<MTLTexture>)Cubemap->GetNativeResource() newTextureViewWithPixelFormat: MTLPixelFormatBGRA8Unorm textureType: MTLTextureType2D levels: NSMakeRange(0, 1) slices: NSMakeRange(OurCubeIndex, 1)];
+		id<MTLTexture> UnderlyingMetalTexture = (id<MTLTexture>)Cubemap->GetNativeResource();
+		id<MTLTexture> OurCubeFaceMetalTexture = [UnderlyingMetalTexture newTextureViewWithPixelFormat: MTLPixelFormatBGRA8Unorm textureType: MTLTextureType2D levels: NSMakeRange(0, 1) slices: NSMakeRange(OurCubeIndex, 1)];
 
 		CIContext* Context = [CIContext context];
 
 		[Context render: RotatedCubefaceImage toMTLTexture: OurCubeFaceMetalTexture commandBuffer: nil bounds: CubefaceImage.extent colorSpace: CubefaceImage.colorSpace];
+
 		[CubefaceImage release];
-//			[RotatedCubefaceImage release];
 		[CubeFaceMetalTexture release];
 		[OurCubeFaceMetalTexture release];
+		[Context release];
 	}
 	
 	virtual void ReleaseRHI() override
