Merging using Dev-VR_to_Release-4.20

diff --git a/Engine/Plugins/Runtime/AppleARKit/Source/AppleARKit/Private/AppleARKitTrackable.cpp b/Engine/Plugins/Runtime/AppleARKit/Source/AppleARKit/Private/AppleARKitTrackable.cpp
--- a/Engine/Plugins/Runtime/AppleARKit/Source/AppleARKit/Private/AppleARKitTrackable.cpp
+++ b/Engine/Plugins/Runtime/AppleARKit/Source/AppleARKit/Private/AppleARKitTrackable.cpp
@@ -7,12 +7,14 @@ UAppleARKitEnvironmentCaptureProbe::UAppleARKitEnvironmentCaptureProbe()
 	: Super()
 	, ARKitEnvironmentTexture(nullptr)
 {
+#if SUPPORTS_ARKIT_2_0
 	ARKitEnvironmentTexture = NewObject<UAppleARKitEnvironmentCaptureProbeTexture>();
 	// Set the base class member since that's what gets used by the non-ARKit specific code
 	EnvironmentCaptureTexture = ARKitEnvironmentTexture;
+#endif
 }
 
-#if PLATFORM_MAC || PLATFORM_IOS
+#if SUPPORTS_ARKIT_2_0
 void UAppleARKitEnvironmentCaptureProbe::UpdateEnvironmentCapture(const TSharedRef<FARSystemBase, ESPMode::ThreadSafe>& InTrackingSystem, uint32 FrameNumber, double InTimestamp, const FTransform& InLocalToTrackingTransform, const FTransform& InAlignmentTransform, FVector InExtent, id<MTLTexture> InMetalTexture)
 {
 	Super::UpdateEnvironmentCapture(InTrackingSystem, FrameNumber, InTimestamp, InLocalToTrackingTransform, InAlignmentTransform, InExtent);
diff --git a/Engine/Plugins/Runtime/AppleARKit/Source/AppleARKit/Public/AppleARKitTrackable.h b/Engine/Plugins/Runtime/AppleARKit/Source/AppleARKit/Public/AppleARKitTrackable.h
--- a/Engine/Plugins/Runtime/AppleARKit/Source/AppleARKit/Public/AppleARKitTrackable.h
+++ b/Engine/Plugins/Runtime/AppleARKit/Source/AppleARKit/Public/AppleARKitTrackable.h
@@ -18,7 +18,7 @@ class APPLEARKIT_API UAppleARKitEnvironmentCaptureProbe :
 public:
 	UAppleARKitEnvironmentCaptureProbe();
 	
-#if PLATFORM_MAC || PLATFORM_IOS
+#if SUPPORTS_ARKIT_2_0
 	/** Updates the current capture with the new metal texture. This will trigger a render resource update if the textures have changed */
 	void UpdateEnvironmentCapture(const TSharedRef<FARSystemBase, ESPMode::ThreadSafe>& InTrackingSystem, uint32 FrameNumber, double Timestamp, const FTransform& InLocalToTrackingTransform, const FTransform& InAlignmentTransform, FVector InExtent, id<MTLTexture> InMetalTexture);
 #endif
