Fixed breakage due to ARKit compiled without 2.0 support

diff --git a/Engine/Plugins/Runtime/AppleARKit/Source/AppleARKit/Private/AppleARKitTrackable.cpp b/Engine/Plugins/Runtime/AppleARKit/Source/AppleARKit/Private/AppleARKitTrackable.cpp
--- a/Engine/Plugins/Runtime/AppleARKit/Source/AppleARKit/Private/AppleARKitTrackable.cpp
+++ b/Engine/Plugins/Runtime/AppleARKit/Source/AppleARKit/Private/AppleARKitTrackable.cpp
@@ -14,7 +14,7 @@ UAppleARKitEnvironmentCaptureProbe::UAppleARKitEnvironmentCaptureProbe()
 #endif
 }
 
-#if SUPPORTS_ARKIT_2_0
+#if PLATFORM_IOS
 void UAppleARKitEnvironmentCaptureProbe::UpdateEnvironmentCapture(const TSharedRef<FARSystemBase, ESPMode::ThreadSafe>& InTrackingSystem, uint32 FrameNumber, double InTimestamp, const FTransform& InLocalToTrackingTransform, const FTransform& InAlignmentTransform, FVector InExtent, id<MTLTexture> InMetalTexture)
 {
 	Super::UpdateEnvironmentCapture(InTrackingSystem, FrameNumber, InTimestamp, InLocalToTrackingTransform, InAlignmentTransform, InExtent);
diff --git a/Engine/Plugins/Runtime/AppleARKit/Source/AppleARKit/Public/AppleARKitTrackable.h b/Engine/Plugins/Runtime/AppleARKit/Source/AppleARKit/Public/AppleARKitTrackable.h
--- a/Engine/Plugins/Runtime/AppleARKit/Source/AppleARKit/Public/AppleARKitTrackable.h
+++ b/Engine/Plugins/Runtime/AppleARKit/Source/AppleARKit/Public/AppleARKitTrackable.h
@@ -18,7 +18,7 @@ class APPLEARKIT_API UAppleARKitEnvironmentCaptureProbe :
 public:
 	UAppleARKitEnvironmentCaptureProbe();
 	
-#if SUPPORTS_ARKIT_2_0
+#if PLATFORM_IOS
 	/** Updates the current capture with the new metal texture. This will trigger a render resource update if the textures have changed */
 	void UpdateEnvironmentCapture(const TSharedRef<FARSystemBase, ESPMode::ThreadSafe>& InTrackingSystem, uint32 FrameNumber, double Timestamp, const FTransform& InLocalToTrackingTransform, const FTransform& InAlignmentTransform, FVector InExtent, id<MTLTexture> InMetalTexture);
 #endif
