Fixed the rotation issue with EARWorldAlignment::Camera

diff --git a/Engine/Plugins/Runtime/AppleARKit/Source/AppleARKit/Private/AppleARKitSystem.cpp b/Engine/Plugins/Runtime/AppleARKit/Source/AppleARKit/Private/AppleARKitSystem.cpp
--- a/Engine/Plugins/Runtime/AppleARKit/Source/AppleARKit/Private/AppleARKitSystem.cpp
+++ b/Engine/Plugins/Runtime/AppleARKit/Source/AppleARKit/Private/AppleARKitSystem.cpp
@@ -224,7 +224,7 @@ bool FAppleARKitSystem::GetCurrentPose(int32 DeviceId, FQuat& OutOrientation, FV
 		// Apply counter-rotation to compensate for mobile device orientation
 		OutOrientation = CurrentTransform.GetRotation();
 		OutPosition = CurrentTransform.GetLocation();
-		
+
 		return true;
 	}
 	else
@@ -250,30 +250,55 @@ bool FAppleARKitSystem::EnumerateTrackedDevices(TArray<int32>& OutDevices, EXRTr
 	return false;
 }
 
-FRotator DeriveTrackingToWorldRotation( EScreenOrientation::Type DeviceOrientation )
+void FAppleARKitSystem::CalcTrackingToWorldRotation()
 {
 	// We rotate the camera to counteract the portrait vs. landscape viewport rotation
-	FRotator DeviceRot = FRotator::ZeroRotator;
-	switch (DeviceOrientation)
+	DerivedTrackingToUnrealRotation = FRotator::ZeroRotator;
+
+	const EARWorldAlignment WorldAlignment = GetSessionConfig().GetWorldAlignment();
+	if (WorldAlignment == EARWorldAlignment::Gravity || WorldAlignment == EARWorldAlignment::GravityAndHeading)
 	{
-		case EScreenOrientation::Portrait:
-			DeviceRot = FRotator(0.0f, 0.0f, -90.0f);
-			break;
-			
-		case EScreenOrientation::PortraitUpsideDown:
-			DeviceRot = FRotator(0.0f, 0.0f, 90.0f);
-			break;
-			
-		default:
-		case EScreenOrientation::LandscapeLeft:
-			break;
-			
-		case EScreenOrientation::LandscapeRight:
-			DeviceRot = FRotator(0.0f, 0.0f, 180.0f);
-			break;
-	};
-	
-	return DeviceRot;
+		switch (DeviceOrientation)
+		{
+			case EScreenOrientation::Portrait:
+				DerivedTrackingToUnrealRotation = FRotator(0.0f, 0.0f, -90.0f);
+				break;
+				
+			case EScreenOrientation::PortraitUpsideDown:
+				DerivedTrackingToUnrealRotation = FRotator(0.0f, 0.0f, 90.0f);
+				break;
+				
+			default:
+			case EScreenOrientation::LandscapeLeft:
+				break;
+				
+			case EScreenOrientation::LandscapeRight:
+				DerivedTrackingToUnrealRotation = FRotator(0.0f, 0.0f, 180.0f);
+				break;
+		}
+	}
+	// Camera aligned which means +X is to the right along the long axis
+	else
+	{
+		switch (DeviceOrientation)
+		{
+			case EScreenOrientation::Portrait:
+				DerivedTrackingToUnrealRotation = FRotator(0.0f, 0.0f, 90.0f);
+				break;
+				
+			case EScreenOrientation::PortraitUpsideDown:
+				DerivedTrackingToUnrealRotation = FRotator(0.0f, 0.0f, -90.0f);
+				break;
+				
+			default:
+			case EScreenOrientation::LandscapeLeft:
+				DerivedTrackingToUnrealRotation = FRotator(0.0f, 0.0f, -180.0f);
+				break;
+				
+			case EScreenOrientation::LandscapeRight:
+				break;
+		}
+	}
 }
 
 void FAppleARKitSystem::UpdateFrame()
@@ -319,10 +344,6 @@ void FAppleARKitSystem::UpdateFrame()
 
 void FAppleARKitSystem::UpdatePoses()
 {
-	if (DeviceOrientation == EScreenOrientation::Unknown)
-	{
-		SetDeviceOrientation( static_cast<EScreenOrientation::Type>(FPlatformMisc::GetDeviceOrientation()) );
-	}
 	UpdateFrame();
 }
 
@@ -1050,7 +1071,7 @@ void FAppleARKitSystem::SetDeviceOrientation( EScreenOrientation::Type InOrienta
 	if (NewOrientation.IsSet() && DeviceOrientation != NewOrientation.GetValue())
 	{
 		DeviceOrientation = NewOrientation.GetValue();
-		DerivedTrackingToUnrealRotation = DeriveTrackingToWorldRotation( DeviceOrientation );
+		CalcTrackingToWorldRotation();
 	}
 }
 
@@ -1073,6 +1094,12 @@ bool FAppleARKitSystem::Run(UARSessionConfig* SessionConfig)
 		LastReceivedFrame = TSharedPtr<FAppleARKitFrame, ESPMode::ThreadSafe>();
 	}
 	
+	// Make sure this is set at session start, because there are timing issues with using only the delegate approach
+	if (DeviceOrientation == EScreenOrientation::Unknown)
+	{
+		SetDeviceOrientation( static_cast<EScreenOrientation::Type>(FPlatformMisc::GetDeviceOrientation()) );
+	}
+
 #if SUPPORTS_ARKIT_1_0
 	if (FAppleARKitAvailability::SupportsARKit10())
 	{
@@ -1298,7 +1325,8 @@ void FAppleARKitSystem::SessionDidAddAnchors_DelegateThread( NSArray<ARAnchor*>*
 	// If this object is valid, we are running a face session and need that code to process things
 	if (FaceARSupport != nullptr)
 	{
-		const TArray<TSharedPtr<FAppleARKitAnchorData>> AnchorList = FaceARSupport->MakeAnchorData(anchors, GameThreadTimestamp, GameThreadFrameNumber);
+		const FRotator& AdjustBy = GetSessionConfig().GetWorldAlignment() == EARWorldAlignment::Camera ? DerivedTrackingToUnrealRotation : FRotator::ZeroRotator;
+		const TArray<TSharedPtr<FAppleARKitAnchorData>> AnchorList = FaceARSupport->MakeAnchorData(anchors, GameThreadTimestamp, GameThreadFrameNumber, AdjustBy);
 		for (TSharedPtr<FAppleARKitAnchorData> NewAnchorData : AnchorList)
 		{
 			auto AddAnchorTask = FSimpleDelegateGraphTask::FDelegate::CreateSP(this, &FAppleARKitSystem::SessionDidAddAnchors_Internal, NewAnchorData.ToSharedRef());
@@ -1327,7 +1355,8 @@ void FAppleARKitSystem::SessionDidUpdateAnchors_DelegateThread( NSArray<ARAnchor
 	// If this object is valid, we are running a face session and need that code to process things
 	if (FaceARSupport != nullptr)
 	{
-		const TArray<TSharedPtr<FAppleARKitAnchorData>> AnchorList = FaceARSupport->MakeAnchorData(anchors, GameThreadTimestamp, GameThreadFrameNumber);
+		const FRotator& AdjustBy = GetSessionConfig().GetWorldAlignment() == EARWorldAlignment::Camera ? DerivedTrackingToUnrealRotation : FRotator::ZeroRotator;
+		const TArray<TSharedPtr<FAppleARKitAnchorData>> AnchorList = FaceARSupport->MakeAnchorData(anchors, GameThreadTimestamp, GameThreadFrameNumber, AdjustBy);
 		for (TSharedPtr<FAppleARKitAnchorData> NewAnchorData : AnchorList)
 		{
 			auto UpdateAnchorTask = FSimpleDelegateGraphTask::FDelegate::CreateSP(this, &FAppleARKitSystem::SessionDidUpdateAnchors_Internal, NewAnchorData.ToSharedRef());
@@ -1373,7 +1402,7 @@ void FAppleARKitSystem::SessionDidAddAnchors_Internal( TSharedRef<FAppleARKitAnc
 	{
 		UpdateFrame();
 	}
-
+	
 	// If this object is valid, we are running a face session and we need to publish LiveLink data on the game thread
 	if (FaceARSupport != nullptr && AnchorData->AnchorType == EAppleAnchorType::FaceAnchor)
 	{
diff --git a/Engine/Plugins/Runtime/AppleARKit/Source/AppleARKit/Private/AppleARKitSystem.h b/Engine/Plugins/Runtime/AppleARKit/Source/AppleARKit/Private/AppleARKitSystem.h
--- a/Engine/Plugins/Runtime/AppleARKit/Source/AppleARKit/Private/AppleARKitSystem.h
+++ b/Engine/Plugins/Runtime/AppleARKit/Source/AppleARKit/Private/AppleARKitSystem.h
@@ -94,6 +94,7 @@ private:
 	void OrientationChanged(const int32 NewOrientation);
 	void UpdatePoses();
 	void UpdateFrame();
+	void CalcTrackingToWorldRotation();
 
 public:
 	// Session delegate callbacks
diff --git a/Engine/Plugins/Runtime/AppleARKit/Source/AppleARKit/Public/AppleARKitConversion.h b/Engine/Plugins/Runtime/AppleARKit/Source/AppleARKit/Public/AppleARKitConversion.h
--- a/Engine/Plugins/Runtime/AppleARKit/Source/AppleARKit/Public/AppleARKitConversion.h
+++ b/Engine/Plugins/Runtime/AppleARKit/Source/AppleARKit/Public/AppleARKitConversion.h
@@ -26,7 +26,7 @@ struct FAppleARKitConversion
 	 *
 	 * Ignores scale.
 	 */
-	static FORCEINLINE FTransform ToFTransform(const matrix_float4x4& RawYUpMatrix)
+	static FORCEINLINE FTransform ToFTransform(const matrix_float4x4& RawYUpMatrix, const FRotator& AdjustBy = FRotator::ZeroRotator)
 	{
 		// Conversion here is as per SteamVRHMD::ToFMatrix
 		FMatrix RawYUpFMatrix(
@@ -41,6 +41,10 @@ struct FAppleARKitConversion
 		// Extract & convert rotation 
 		FQuat RawRotation( RawYUpFMatrix );
 		FQuat Rotation( -RawRotation.Z, RawRotation.X, RawRotation.Y, -RawRotation.W );
+		if (!AdjustBy.IsNearlyZero())
+		{
+			Rotation = FQuat(AdjustBy) * Rotation;
+		}
 
 		return FTransform( Rotation, Translation );
 	}
diff --git a/Engine/Plugins/Runtime/AppleARKit/Source/AppleARKit/Public/AppleARKitFaceSupport.h b/Engine/Plugins/Runtime/AppleARKit/Source/AppleARKit/Public/AppleARKitFaceSupport.h
--- a/Engine/Plugins/Runtime/AppleARKit/Source/AppleARKit/Public/AppleARKitFaceSupport.h
+++ b/Engine/Plugins/Runtime/AppleARKit/Source/AppleARKit/Public/AppleARKitFaceSupport.h
@@ -26,10 +26,11 @@ public:
 	 * @param NewAnchors the list of anchors to convert to our intermediate format
 	 * @param Timestamp the timestamp of this update
 	 * @param FrameNumber the frame number for this update
+	 * @param AdjustBy the additional rotation to apply to put the rotation in the proper space (camera alignment only)
 	 *
 	 * @return the set of face anchors to dispatch
 	 */
-	virtual TArray<TSharedPtr<FAppleARKitAnchorData>> MakeAnchorData(NSArray<ARAnchor*>* NewAnchors, double Timestamp, uint32 FrameNumber) { return TArray<TSharedPtr<FAppleARKitAnchorData>>(); }
+	virtual TArray<TSharedPtr<FAppleARKitAnchorData>> MakeAnchorData(NSArray<ARAnchor*>* NewAnchors, double Timestamp, uint32 FrameNumber, const FRotator& AdjustBy) { return TArray<TSharedPtr<FAppleARKitAnchorData>>(); }
 
 	/**
 	 * Publishes any face AR data that needs to be sent to LiveLink. Done as a separate step because MakeAnchorData is called
diff --git a/Engine/Plugins/Runtime/AppleARKitFaceSupport/Source/AppleARKitFaceSupport/Private/AppleARKitFaceSupportImpl.cpp b/Engine/Plugins/Runtime/AppleARKitFaceSupport/Source/AppleARKitFaceSupport/Private/AppleARKitFaceSupportImpl.cpp
--- a/Engine/Plugins/Runtime/AppleARKitFaceSupport/Source/AppleARKitFaceSupport/Private/AppleARKitFaceSupportImpl.cpp
+++ b/Engine/Plugins/Runtime/AppleARKitFaceSupport/Source/AppleARKitFaceSupport/Private/AppleARKitFaceSupportImpl.cpp
@@ -12,7 +12,7 @@ DECLARE_STATS_GROUP(TEXT("AppleARKitFaceSupport"), STATGROUP_APPLEARKITFACE, STA
 
 #if SUPPORTS_ARKIT_1_0
 
-static TSharedPtr<FAppleARKitAnchorData> MakeAnchorData(ARAnchor* Anchor)
+static TSharedPtr<FAppleARKitAnchorData> MakeAnchorData(ARAnchor* Anchor, const FRotator& AdjustBy)
 {
     TSharedPtr<FAppleARKitAnchorData> NewAnchor;
     if ([Anchor isKindOfClass:[ARFaceAnchor class]])
@@ -25,15 +25,15 @@ static TSharedPtr<FAppleARKitAnchorData> MakeAnchorData(ARAnchor* Anchor)
 #if SUPPORTS_ARKIT_2_0
 		if (FAppleARKitAvailability::SupportsARKit20())
 		{
-			LeftEyeTransform = FAppleARKitConversion::ToFTransform(FaceAnchor.leftEyeTransform);
-			RightEyeTransform = FAppleARKitConversion::ToFTransform(FaceAnchor.rightEyeTransform);
+			LeftEyeTransform = FAppleARKitConversion::ToFTransform(FaceAnchor.leftEyeTransform, AdjustBy);
+			RightEyeTransform = FAppleARKitConversion::ToFTransform(FaceAnchor.rightEyeTransform, AdjustBy);
 			LookAtTarget = FAppleARKitConversion::ToFVector(FaceAnchor.lookAtPoint);
 		}
 #endif
 		NewAnchor = MakeShared<FAppleARKitAnchorData>(
 			FAppleARKitConversion::ToFGuid(FaceAnchor.identifier),
-			FAppleARKitConversion::ToFTransform(FaceAnchor.transform),
-			ToBlendShapeMap(FaceAnchor.blendShapes, FAppleARKitConversion::ToFTransform(FaceAnchor.transform), LeftEyeTransform, RightEyeTransform),
+			FAppleARKitConversion::ToFTransform(FaceAnchor.transform, AdjustBy),
+			ToBlendShapeMap(FaceAnchor.blendShapes, FAppleARKitConversion::ToFTransform(FaceAnchor.transform, AdjustBy), LeftEyeTransform, RightEyeTransform),
 			ToVertexBuffer(FaceAnchor.geometry.vertices, FaceAnchor.geometry.vertexCount),
 			LeftEyeTransform,
 			RightEyeTransform,
@@ -93,13 +93,13 @@ ARConfiguration* FAppleARKitFaceSupport::ToARConfiguration(UARSessionConfig* Ses
 	return SessionConfiguration;
 }
 
-TArray<TSharedPtr<FAppleARKitAnchorData>> FAppleARKitFaceSupport::MakeAnchorData(NSArray<ARAnchor*>* Anchors, double Timestamp, uint32 FrameNumber)
+TArray<TSharedPtr<FAppleARKitAnchorData>> FAppleARKitFaceSupport::MakeAnchorData(NSArray<ARAnchor*>* Anchors, double Timestamp, uint32 FrameNumber, const FRotator& AdjustBy)
 {
 	TArray<TSharedPtr<FAppleARKitAnchorData>> AnchorList;
 
 	for (ARAnchor* Anchor in Anchors)
 	{
-		TSharedPtr<FAppleARKitAnchorData> AnchorData = ::MakeAnchorData(Anchor);
+		TSharedPtr<FAppleARKitAnchorData> AnchorData = ::MakeAnchorData(Anchor, AdjustBy);
 		if (AnchorData.IsValid())
 		{
 			AnchorList.Add(AnchorData);
diff --git a/Engine/Plugins/Runtime/AppleARKitFaceSupport/Source/AppleARKitFaceSupport/Private/AppleARKitFaceSupportImpl.h b/Engine/Plugins/Runtime/AppleARKitFaceSupport/Source/AppleARKitFaceSupport/Private/AppleARKitFaceSupportImpl.h
--- a/Engine/Plugins/Runtime/AppleARKitFaceSupport/Source/AppleARKitFaceSupport/Private/AppleARKitFaceSupportImpl.h
+++ b/Engine/Plugins/Runtime/AppleARKitFaceSupport/Source/AppleARKitFaceSupport/Private/AppleARKitFaceSupportImpl.h
@@ -22,7 +22,7 @@ public:
 private:
 #if SUPPORTS_ARKIT_1_0
 	// ~IAppleARKitFaceSupport
-	virtual TArray<TSharedPtr<FAppleARKitAnchorData>> MakeAnchorData(NSArray<ARAnchor*>* NewAnchors, double Timestamp, uint32 FrameNumber) override;
+	virtual TArray<TSharedPtr<FAppleARKitAnchorData>> MakeAnchorData(NSArray<ARAnchor*>* NewAnchors, double Timestamp, uint32 FrameNumber, const FRotator& AdjustBy) override;
 	virtual ARConfiguration* ToARConfiguration(UARSessionConfig* InSessionConfig) override;
 	virtual void PublishLiveLinkData(TSharedPtr<FAppleARKitAnchorData> Anchor, double Timestamp, uint32 FrameNumber) override;
 	// ~IAppleARKitFaceSupport
