Added bIsTracked to ARTrackedImage and ARFaceGeometry since they both support the ARTrackable interface

diff --git a/Engine/Plugins/Runtime/AppleARKit/Source/AppleARKit/Private/AppleARKitSystem.cpp b/Engine/Plugins/Runtime/AppleARKit/Source/AppleARKit/Private/AppleARKitSystem.cpp
--- a/Engine/Plugins/Runtime/AppleARKit/Source/AppleARKit/Private/AppleARKitSystem.cpp
+++ b/Engine/Plugins/Runtime/AppleARKit/Source/AppleARKit/Private/AppleARKitSystem.cpp
@@ -1280,6 +1280,7 @@ static TSharedPtr<FAppleARKitAnchorData> MakeAnchorData( ARAnchor* Anchor )
 			EAppleAnchorType::ImageAnchor,
 			FString(ImageAnchor.referenceImage.name)
 		);
+		NewAnchor->bIsTracked = ImageAnchor.isTracked;
 	}
 #endif
 //@joeg -- Added environmental texture probe support
@@ -1434,6 +1435,7 @@ void FAppleARKitSystem::SessionDidAddAnchors_Internal( TSharedRef<FAppleARKitAnc
 			UARFaceGeometry* NewGeo = NewObject<UARFaceGeometry>();
 //@joeg -- Eye tracking support
 			NewGeo->UpdateFaceGeometry(SharedThis(this), GameThreadFrameNumber, GameThreadTimestamp, AnchorData->Transform, GetAlignmentTransform(), AnchorData->BlendShapes, AnchorData->FaceVerts, AnchorData->FaceIndices, AnchorData->LeftEyeTransform, AnchorData->RightEyeTransform, AnchorData->LookAtTarget);
+			NewGeo->bIsTracked = true;
 			NewGeometry = NewGeo;
 			break;
 		}
@@ -1535,6 +1537,7 @@ void FAppleARKitSystem::SessionDidUpdateAnchors_Internal( TSharedRef<FAppleARKit
 				{
 //@joeg -- Eye tracking support
 					FaceGeo->UpdateFaceGeometry(SharedThis(this), GameThreadFrameNumber, GameThreadTimestamp, AnchorData->Transform, GetAlignmentTransform(), AnchorData->BlendShapes, AnchorData->FaceVerts, AnchorData->FaceIndices, AnchorData->LeftEyeTransform, AnchorData->RightEyeTransform, AnchorData->LookAtTarget);
+					FaceGeo->bIsTracked = AnchorData->bIsTracked;
 					for (UARPin* Pin : PinsToUpdate)
 					{
 						const FTransform Pin_LocalToTrackingTransform_PostUpdate = Pin->GetLocalToTrackingTransform_NoAlignment() * AnchorDeltaTransform;
@@ -1553,6 +1556,7 @@ void FAppleARKitSystem::SessionDidUpdateAnchors_Internal( TSharedRef<FAppleARKit
 					ensure(CandidateImage != nullptr);
 
 					ImageAnchor->UpdateTrackedGeometry(SharedThis(this), GameThreadFrameNumber, GameThreadTimestamp, AnchorData->Transform, GetAlignmentTransform(), *CandidateImage);
+					ImageAnchor->bIsTracked = AnchorData->bIsTracked;
 					for (UARPin* Pin : PinsToUpdate)
 					{
 						const FTransform Pin_LocalToTrackingTransform_PostUpdate = Pin->GetLocalToTrackingTransform_NoAlignment() * AnchorDeltaTransform;
diff --git a/Engine/Plugins/Runtime/AppleARKit/Source/AppleARKit/Public/AppleARKitConversion.h b/Engine/Plugins/Runtime/AppleARKit/Source/AppleARKit/Public/AppleARKitConversion.h
--- a/Engine/Plugins/Runtime/AppleARKit/Source/AppleARKit/Public/AppleARKitConversion.h
+++ b/Engine/Plugins/Runtime/AppleARKit/Source/AppleARKit/Public/AppleARKitConversion.h
@@ -428,6 +428,7 @@ struct FAppleARKitAnchorData
 		: Transform( InTransform )
 		, AnchorType( EAppleAnchorType::Anchor )
 		, AnchorGUID( InAnchorGuid )
+		, bIsTracked(false)
 	{
 	}
 
@@ -498,6 +499,9 @@ struct FAppleARKitAnchorData
 	FTransform LeftEyeTransform;
 	FTransform RightEyeTransform;
 	FVector LookAtTarget;
+
+	/** Only valid for tracked real world objects (face, images) */
+	bool bIsTracked;
 };
 
 namespace ARKitUtil
diff --git a/Engine/Plugins/Runtime/AppleARKitFaceSupport/Source/AppleARKitFaceSupport/Private/AppleARKitFaceSupportImpl.cpp b/Engine/Plugins/Runtime/AppleARKitFaceSupport/Source/AppleARKitFaceSupport/Private/AppleARKitFaceSupportImpl.cpp
--- a/Engine/Plugins/Runtime/AppleARKitFaceSupport/Source/AppleARKitFaceSupport/Private/AppleARKitFaceSupportImpl.cpp
+++ b/Engine/Plugins/Runtime/AppleARKitFaceSupport/Source/AppleARKitFaceSupport/Private/AppleARKitFaceSupportImpl.cpp
@@ -44,6 +44,7 @@ static TSharedPtr<FAppleARKitAnchorData> MakeAnchorData(ARAnchor* Anchor, const
         {
             FAppleARKitAnchorData::FaceIndices = To32BitIndexBuffer(FaceAnchor.geometry.triangleIndices, FaceAnchor.geometry.triangleCount * 3);
         }
+		NewAnchor->bIsTracked = FaceAnchor.isTracked;
     }
 
     return NewAnchor;
diff --git a/Engine/Source/Runtime/AugmentedReality/Public/ARTrackable.h b/Engine/Source/Runtime/AugmentedReality/Public/ARTrackable.h
--- a/Engine/Source/Runtime/AugmentedReality/Public/ARTrackable.h
+++ b/Engine/Source/Runtime/AugmentedReality/Public/ARTrackable.h
@@ -143,6 +143,10 @@ public:
 	UFUNCTION(BlueprintPure, Category = "AR AugmentedReality|Image Detection")
 	UARCandidateImage* GetDetectedImage() const { return DetectedImage; };
 
+	/** Whether the image is currently being tracked by the AR system */
+	UPROPERTY(BlueprintReadOnly, Category="AR AugmentedReality|Face Geometry")
+	bool bIsTracked;
+
 private:
 	/** The candidate image that was detected in the scene */
 	UPROPERTY()
@@ -275,11 +279,14 @@ public:
 	UFUNCTION(BlueprintPure, Category="AR AugmentedReality|Face Geometry")
 	FTransform GetWorldSpaceEyeTransform(EAREye Eye) const;
 
-//@joeg -- Eye tracking support
 	/** The target the eyes are looking at */
 	UPROPERTY(BlueprintReadOnly, Category="AR AugmentedReality|Face Geometry")
 	FVector LookAtTarget;
 
+	/** Whether the face is currently being tracked by the AR system */
+	UPROPERTY(BlueprintReadOnly, Category="AR AugmentedReality|Face Geometry")
+	bool bIsTracked;
+
 private:
 	UPROPERTY()
 	TMap<EARFaceBlendShape, float> BlendShapes;
