Fixed missing SUPPORTS_ARKIT_2_0 #ifdef around the image isTracked code

diff --git a/Engine/Plugins/Runtime/AppleARKit/Source/AppleARKit/Private/AppleARKitSystem.cpp b/Engine/Plugins/Runtime/AppleARKit/Source/AppleARKit/Private/AppleARKitSystem.cpp
--- a/Engine/Plugins/Runtime/AppleARKit/Source/AppleARKit/Private/AppleARKitSystem.cpp
+++ b/Engine/Plugins/Runtime/AppleARKit/Source/AppleARKit/Private/AppleARKitSystem.cpp
@@ -1280,7 +1280,9 @@ static TSharedPtr<FAppleARKitAnchorData> MakeAnchorData( ARAnchor* Anchor )
 			EAppleAnchorType::ImageAnchor,
 			FString(ImageAnchor.referenceImage.name)
 		);
+#if SUPPORTS_ARKIT_2_0
 		NewAnchor->bIsTracked = ImageAnchor.isTracked;
+#endif
 	}
 #endif
 //@joeg -- Added environmental texture probe support
