Fix for AR Camera Pass Through Materials

diff --git a/Engine/Plugins/Runtime/AR/Apple/AppleARKit/Source/AppleARKit/Private/AppleARKitVideoOverlay.cpp b/Engine/Plugins/Runtime/AR/Apple/AppleARKit/Source/AppleARKit/Private/AppleARKitVideoOverlay.cpp
--- a/Engine/Plugins/Runtime/AR/Apple/AppleARKit/Source/AppleARKit/Private/AppleARKitVideoOverlay.cpp
+++ b/Engine/Plugins/Runtime/AR/Apple/AppleARKit/Source/AppleARKit/Private/AppleARKitVideoOverlay.cpp
@@ -240,6 +240,9 @@ void FAppleARKitVideoOverlay::UpdateVideoTexture_RenderThread(FRHICommandListImm
 
 			FExternalTextureRegistry::Get().RegisterExternalTexture(ARKitPassthroughCameraExternalTextureYGuid, VideoTextureY, SamplerStateRHI);
 			FExternalTextureRegistry::Get().RegisterExternalTexture(ARKitPassthroughCameraExternalTextureCbCrGuid, VideoTextureCbCr, SamplerStateRHI);
+			
+			//Make sure AR camera pass through materials are updated properly
+			FMaterialRenderProxy::UpdateDeferredCachedUniformExpressions();
 
 			CFRelease(Frame.CapturedYImage);
 			CFRelease(Frame.CapturedCbCrImage);
