Allow projects to opt-in to ARKit support without modifying the requried device capabilitites plist entry. This allows existing or new projects to add optional ARKit features.

Required for project k.

diff --git a/Engine/Config/BaseEngine.ini b/Engine/Config/BaseEngine.ini
--- a/Engine/Config/BaseEngine.ini
+++ b/Engine/Config/BaseEngine.ini
@@ -2261,6 +2261,9 @@ PakEncryptionRequired=True
 [/Script/AppleARKit.AppleARKitSettings]
 bEnableLiveLinkForFaceTracking=true
 LiveLinkPublishingPort=11111
+; by default arkit will be added to UIRequiredDeviceCapabilities and devices must support
+; arkit. If false games must check that arkit is supported before enabling arkit features
+bRequireDeviceSupportsARKit=true
 
 [/Script/Engine.RendererSettings]
 r.GPUCrashDebugging=false
diff --git a/Engine/Plugins/Runtime/AR/Apple/AppleARKit/Source/AppleARKit/AppleARKit_IOS_UPL.xml b/Engine/Plugins/Runtime/AR/Apple/AppleARKit/Source/AppleARKit/AppleARKit_IOS_UPL.xml
--- a/Engine/Plugins/Runtime/AR/Apple/AppleARKit/Source/AppleARKit/AppleARKit_IOS_UPL.xml
+++ b/Engine/Plugins/Runtime/AR/Apple/AppleARKit/Source/AppleARKit/AppleARKit_IOS_UPL.xml
@@ -11,6 +11,9 @@
             <string>The camera is used for augmenting reality.</string>
         </addElements>
 
+		<setBoolFromProperty result="bRequireDeviceSupportsARKit" ini="Engine" section="/Script/AppleARKit.AppleARKitSettings" property="bRequireDeviceSupportsARKit" default="true"/>
+		<if condition="bRequireDeviceSupportsARKit">
+			<true>
         <loopElements tag="dict">
             <setBool result="bReqKey" value="false"/>
 
@@ -39,6 +42,8 @@
 
             </loopElements>
         </loopElements>
+			</true>
+		</if>
     </iosPListUpdates>
 
 </root>
\ No newline at end of file
