UE-76026 //UE4/Main - Compile UE4Game iOS - error: use of undeclared identifier 'ToEARObjectClassification'
-Fixed build break

diff --git a/Engine/Plugins/Runtime/AR/Apple/AppleARKit/Source/AppleARKit/Private/AppleARKitSystem.cpp b/Engine/Plugins/Runtime/AR/Apple/AppleARKit/Source/AppleARKit/Private/AppleARKitSystem.cpp
--- a/Engine/Plugins/Runtime/AR/Apple/AppleARKit/Source/AppleARKit/Private/AppleARKitSystem.cpp
+++ b/Engine/Plugins/Runtime/AR/Apple/AppleARKit/Source/AppleARKit/Private/AppleARKitSystem.cpp
@@ -1416,7 +1416,7 @@ static TSharedPtr<FAppleARKitAnchorData> MakeAnchorData( ARAnchor* Anchor, doubl
 			NewAnchor->ObjectClassification = FAppleARKitConversion::ToEARObjectClassification(PlaneAnchor.classification);
 		}
 #else
-		NewAnchor->ObjectClassification = ToEARObjectClassification::Unknown;
+		NewAnchor->ObjectClassification = EARObjectClassification::Unknown;
 #endif
 	}
 #if SUPPORTS_ARKIT_1_5
