Merging using Dev-VR_to_Release-4.21

diff --git a/Engine/Plugins/Runtime/AppleARKit/Source/AppleARKit/Private/AppleARKitTextures.cpp b/Engine/Plugins/Runtime/AppleARKit/Source/AppleARKit/Private/AppleARKitTextures.cpp
--- a/Engine/Plugins/Runtime/AppleARKit/Source/AppleARKit/Private/AppleARKitTextures.cpp
+++ b/Engine/Plugins/Runtime/AppleARKit/Source/AppleARKit/Private/AppleARKitTextures.cpp
@@ -267,33 +267,31 @@ public:
 
 	void CopyCubeFace(id<MTLTexture> MetalTexture, FTextureCubeRHIRef Cubemap, uint32 Rotation, int32 MetalCubeIndex, int32 OurCubeIndex)
 	{
-		// Rotate the image we need to get a view into the face as a new slice
-		id<MTLTexture> CubeFaceMetalTexture = [MetalTexture newTextureViewWithPixelFormat: MTLPixelFormatBGRA8Unorm textureType: MTLTextureType2D levels: NSMakeRange(0, 1) slices: NSMakeRange(MetalCubeIndex, 1)];
-		CIImage* CubefaceImage = [[CIImage alloc] initWithMTLTexture: CubeFaceMetalTexture options: nil];
-		CIImage* RotatedCubefaceImage = [CubefaceImage imageByApplyingOrientation: Rotation];
-		CIImage* ImageTransform = nullptr;
-		if (Rotation != kCGImagePropertyOrientationUp)
+		@autoreleasepool
 		{
-			ImageTransform = RotatedCubefaceImage;
+			// Rotate the image we need to get a view into the face as a new slice
+			id<MTLTexture> CubeFaceMetalTexture = [MetalTexture newTextureViewWithPixelFormat: MTLPixelFormatBGRA8Unorm textureType: MTLTextureType2D levels: NSMakeRange(0, 1) slices: NSMakeRange(MetalCubeIndex, 1)];
+			CIImage* CubefaceImage = [[CIImage alloc] initWithMTLTexture: CubeFaceMetalTexture options: nil];
+			CIImage* RotatedCubefaceImage = [CubefaceImage imageByApplyingOrientation: Rotation];
+			CIImage* ImageTransform = nullptr;
+			if (Rotation != kCGImagePropertyOrientationUp)
+			{
+				ImageTransform = RotatedCubefaceImage;
+			}
+			else
+			{
+				// We don't need to rotate it so just use a copy instead
+				ImageTransform = CubefaceImage;
+			}
+
+			// Make a new view into our texture and directly render to that to avoid the CPU copy
+			id<MTLTexture> UnderlyingMetalTexture = (id<MTLTexture>)Cubemap->GetNativeResource();
+			id<MTLTexture> OurCubeFaceMetalTexture = [UnderlyingMetalTexture newTextureViewWithPixelFormat: MTLPixelFormatBGRA8Unorm textureType: MTLTextureType2D levels: NSMakeRange(0, 1) slices: NSMakeRange(OurCubeIndex, 1)];
+
+			CIContext* Context = [CIContext context];
+
+			[Context render: RotatedCubefaceImage toMTLTexture: OurCubeFaceMetalTexture commandBuffer: nil bounds: CubefaceImage.extent colorSpace: CubefaceImage.colorSpace];
 		}
-		else
-		{
-			// We don't need to rotate it so just use a copy instead
-			ImageTransform = CubefaceImage;
-		}
-
-		// Make a new view into our texture and directly render to that to avoid the CPU copy
-		id<MTLTexture> UnderlyingMetalTexture = (id<MTLTexture>)Cubemap->GetNativeResource();
-		id<MTLTexture> OurCubeFaceMetalTexture = [UnderlyingMetalTexture newTextureViewWithPixelFormat: MTLPixelFormatBGRA8Unorm textureType: MTLTextureType2D levels: NSMakeRange(0, 1) slices: NSMakeRange(OurCubeIndex, 1)];
-
-		CIContext* Context = [CIContext context];
-
-		[Context render: RotatedCubefaceImage toMTLTexture: OurCubeFaceMetalTexture commandBuffer: nil bounds: CubefaceImage.extent colorSpace: CubefaceImage.colorSpace];
-
-		[CubefaceImage release];
-		[CubeFaceMetalTexture release];
-		[OurCubeFaceMetalTexture release];
-		[Context release];
 	}
 	
 	virtual void ReleaseRHI() override
