Merging using Dev-VR_to_Release-4.21

diff --git a/Engine/Plugins/Runtime/AppleARKit/Source/AppleARKit/Private/AppleARKitSystem.cpp b/Engine/Plugins/Runtime/AppleARKit/Source/AppleARKit/Private/AppleARKitSystem.cpp
--- a/Engine/Plugins/Runtime/AppleARKit/Source/AppleARKit/Private/AppleARKitSystem.cpp
+++ b/Engine/Plugins/Runtime/AppleARKit/Source/AppleARKit/Private/AppleARKitSystem.cpp
@@ -18,6 +18,7 @@
 #include "ARLightEstimate.h"
 #include "ARTraceResult.h"
 #include "ARPin.h"
+#include "HAL/ThreadSafeCounter.h"
 
 // To separate out the face ar library linkage from standard ar apps
 #include "AppleARKitFaceSupport.h"
@@ -32,6 +33,32 @@
 	#pragma clang diagnostic ignored "-Wpartial-availability"
 #endif
 
+// Copied from IOSPlatformProcess because it's not accessible by external code
+#define GAME_THREAD_PRIORITY 47
+#define RENDER_THREAD_PRIORITY 45
+
+#if PLATFORM_IOS && !PLATFORM_TVOS
+// Copied from IOSPlatformProcess because it's not accessible by external code
+static void SetThreadPriority(int32 Priority)
+{
+	struct sched_param Sched;
+	FMemory::Memzero(&Sched, sizeof(struct sched_param));
+	
+	// Read the current priority and policy
+	int32 CurrentPolicy = SCHED_RR;
+	pthread_getschedparam(pthread_self(), &CurrentPolicy, &Sched);
+	
+	// Set the new priority and policy (apple recommended FIFO for the two main non-working threads)
+	int32 Policy = SCHED_FIFO;
+	Sched.sched_priority = Priority;
+	pthread_setschedparam(pthread_self(), Policy, &Sched);
+}
+#else
+static void SetThreadPriority(int32 Priority)
+{
+	// Ignored
+}
+#endif
 
 //
 //  FAppleARKitXRCamera
@@ -45,6 +72,11 @@ public:
 	, ARKitSystem( InTrackingSystem )
 	{}
 	
+	void AdjustThreadPriority(int32 NewPriority)
+	{
+		ThreadPriority.Set(NewPriority);
+	}
+	
 private:
 	//~ FDefaultXRCamera
 	void OverrideFOV(float& InOutFOV)
@@ -84,6 +116,12 @@ private:
 	
 	virtual void PreRenderView_RenderThread(FRHICommandListImmediate& RHICmdList, FSceneView& InView) override
 	{
+		// Adjust our thread priority if requested
+		if (LastThreadPriority.GetValue() != ThreadPriority.GetValue())
+		{
+			SetThreadPriority(ThreadPriority.GetValue());
+			LastThreadPriority.Set(ThreadPriority.GetValue());
+		}
 		FDefaultXRCamera::PreRenderView_RenderThread(RHICmdList, InView);
 	}
 	
@@ -140,6 +178,10 @@ private:
 private:
 	FAppleARKitSystem& ARKitSystem;
 	FAppleARKitVideoOverlay VideoOverlay;
+	
+	// Thread priority support
+	FThreadSafeCounter ThreadPriority;
+	FThreadSafeCounter LastThreadPriority;
 };
 
 //
@@ -1197,6 +1239,23 @@ bool FAppleARKitSystem::Run(UARSessionConfig* SessionConfig)
 			// Pass to session delegate to use for Metal texture creation
 			[Delegate setMetalTextureCache : MetalTextureCache];
 		}
+		
+#if PLATFORM_IOS && !PLATFORM_TVOS
+		// Check if we need to adjust the priorities to allow ARKit to have more CPU time
+		if (GetDefault<UAppleARKitSettings>()->bAdjustThreadPrioritiesDuringARSession)
+		{
+			int32 GameOverride = GetDefault<UAppleARKitSettings>()->GameThreadPriorityOverride;
+			int32 RenderOverride = GetDefault<UAppleARKitSettings>()->RenderThreadPriorityOverride;
+			SetThreadPriority(GameOverride);
+			if (XRCamera.IsValid())
+			{
+				FAppleARKitXRCamera* Camera = (FAppleARKitXRCamera*)XRCamera.Get();
+				Camera->AdjustThreadPriority(RenderOverride);
+			}
+			
+			UE_LOG(LogAppleARKit, Log, TEXT("Overriding thread priorities: Game Thread (%d), Render Thread (%d)"), GameOverride, RenderOverride);
+		}
+#endif
 
 		UE_LOG(LogAppleARKit, Log, TEXT("Starting session: %p with options %d"), this, options);
 
@@ -1248,6 +1307,21 @@ bool FAppleARKitSystem::Pause()
 		}
 	}
 	
+#if PLATFORM_IOS && !PLATFORM_TVOS
+	// Check if we need to adjust the priorities to allow ARKit to have more CPU time
+	if (GetDefault<UAppleARKitSettings>()->bAdjustThreadPrioritiesDuringARSession)
+	{
+		SetThreadPriority(GAME_THREAD_PRIORITY);
+		if (XRCamera.IsValid())
+		{
+			FAppleARKitXRCamera* Camera = (FAppleARKitXRCamera*)XRCamera.Get();
+			Camera->AdjustThreadPriority(RENDER_THREAD_PRIORITY);
+		}
+		
+		UE_LOG(LogAppleARKit, Log, TEXT("Restoring thread priorities: Game Thread (%d), Render Thread (%d)"), GAME_THREAD_PRIORITY, RENDER_THREAD_PRIORITY);
+}
+#endif
+	
 #endif
 	
 	// Set running state
diff --git a/Engine/Plugins/Runtime/AppleARKit/Source/AppleARKit/Public/AppleARKitSettings.h b/Engine/Plugins/Runtime/AppleARKit/Source/AppleARKit/Public/AppleARKitSettings.h
--- a/Engine/Plugins/Runtime/AppleARKit/Source/AppleARKit/Public/AppleARKitSettings.h
+++ b/Engine/Plugins/Runtime/AppleARKit/Source/AppleARKit/Public/AppleARKitSettings.h
@@ -15,10 +15,13 @@ class APPLEARKIT_API UAppleARKitSettings : public UObject
 
 public:
 	UAppleARKitSettings()
-			: bEnableLiveLinkForFaceTracking(false)
-			, LiveLinkPublishingPort(11111)
-			, DefaultFaceTrackingLiveLinkSubjectName(FName("iPhoneXFaceAR"))
-			, DefaultFaceTrackingDirection(EARFaceTrackingDirection::FaceRelative)
+		: bEnableLiveLinkForFaceTracking(false)
+		, LiveLinkPublishingPort(11111)
+		, DefaultFaceTrackingLiveLinkSubjectName(FName("iPhoneXFaceAR"))
+		, DefaultFaceTrackingDirection(EARFaceTrackingDirection::FaceRelative)
+		, bAdjustThreadPrioritiesDuringARSession(false)
+		, GameThreadPriorityOverride(47)
+		, RenderThreadPriorityOverride(45)
 	{
 	}
 
@@ -37,4 +40,16 @@ public:
 	/** The default tracking to use when tracking face blend shapes (face relative or mirrored). Defaults to face relative */
 	UPROPERTY(Config, EditAnywhere, Category="AR Settings")
 	EARFaceTrackingDirection DefaultFaceTrackingDirection;
+
+	/** Whether to adjust thread priorities during an AR session or not */
+	UPROPERTY(Config, EditAnywhere, Category="AR Settings")
+	bool bAdjustThreadPrioritiesDuringARSession;
+	
+	/** The game thread priority to change to when an AR session is running, default is 47 */
+	UPROPERTY(Config, EditAnywhere, Category="AR Settings")
+	int32 GameThreadPriorityOverride;
+	
+	/** The render thread priority to change to when an AR session is running, default is 45 */
+	UPROPERTY(Config, EditAnywhere, Category="AR Settings")
+	int32 RenderThreadPriorityOverride;
 };
