Merging using Dev-VR_to_Release-4.21

diff --git a/Engine/Plugins/Runtime/AppleARKit/Source/AppleARKit/Private/AppleARKitSystem.cpp b/Engine/Plugins/Runtime/AppleARKit/Source/AppleARKit/Private/AppleARKitSystem.cpp
--- a/Engine/Plugins/Runtime/AppleARKit/Source/AppleARKit/Private/AppleARKitSystem.cpp
+++ b/Engine/Plugins/Runtime/AppleARKit/Source/AppleARKit/Private/AppleARKitSystem.cpp
@@ -843,6 +843,7 @@ TArray<FVector> FAppleARKitSystem::OnGetPointCloud() const
 {
 	TArray<FVector> PointCloud;
 	
+#if SUPPORTS_ARKIT_1_0
 	if (GameThreadFrame.IsValid())
 	{
 		ARFrame* InARFrame = (ARFrame*)GameThreadFrame->NativeFrame;
@@ -858,6 +859,7 @@ TArray<FVector> FAppleARKitSystem::OnGetPointCloud() const
 			}
 		}
 	}
+#endif
 	return PointCloud;
 }
 
