Merging using Dev-VR_to_Release-4.21

diff --git a/Engine/Plugins/Runtime/AppleARKit/Source/AppleARKit/Private/AppleARKitConversion.cpp b/Engine/Plugins/Runtime/AppleARKit/Source/AppleARKit/Private/AppleARKitConversion.cpp
--- a/Engine/Plugins/Runtime/AppleARKit/Source/AppleARKit/Private/AppleARKitConversion.cpp
+++ b/Engine/Plugins/Runtime/AppleARKit/Source/AppleARKit/Private/AppleARKitConversion.cpp
@@ -103,9 +103,7 @@ NSSet* FAppleARKitConversion::InitImageDetection(UARSessionConfig* SessionConfig
 			if (ConvertedImage != nullptr)
 			{
 				float ImageWidth = (float)Candidate->GetPhysicalWidth() / 100.f;
-				CGImagePropertyOrientation Orientation = Candidate->GetOrientation() == EARCandidateImageOrientation::Landscape ? kCGImagePropertyOrientationRight : kCGImagePropertyOrientationUp;
-
-				ARReferenceImage* ReferenceImage = [[[ARReferenceImage alloc] initWithCGImage: ConvertedImage orientation: Orientation physicalWidth: ImageWidth] autorelease];
+				ARReferenceImage* ReferenceImage = [[[ARReferenceImage alloc] initWithCGImage: ConvertedImage orientation: kCGImagePropertyOrientationUp physicalWidth: ImageWidth] autorelease];
 				ReferenceImage.name = Candidate->GetFriendlyName().GetNSString();
 				[ConvertedImageSet addObject: ReferenceImage];
 			}
