Merging using Dev-VR_to_Release-4.21

diff --git a/Engine/Plugins/Runtime/AppleARKit/Source/AppleARKit/Private/AppleARKitSystem.cpp b/Engine/Plugins/Runtime/AppleARKit/Source/AppleARKit/Private/AppleARKitSystem.cpp
--- a/Engine/Plugins/Runtime/AppleARKit/Source/AppleARKit/Private/AppleARKitSystem.cpp
+++ b/Engine/Plugins/Runtime/AppleARKit/Source/AppleARKit/Private/AppleARKitSystem.cpp
@@ -33,6 +33,13 @@
 	#pragma clang diagnostic ignored "-Wpartial-availability"
 #endif
 
+DECLARE_CYCLE_STAT(TEXT("SessionDidUpdateFrame_DelegateThread"), STAT_FAppleARKitSystem_SessionUpdateFrame, STATGROUP_ARKIT);
+DECLARE_CYCLE_STAT(TEXT("SessionDidAddAnchors_DelegateThread"), STAT_FAppleARKitSystem_SessionDidAddAnchors, STATGROUP_ARKIT);
+DECLARE_CYCLE_STAT(TEXT("SessionDidUpdateAnchors_DelegateThread"), STAT_FAppleARKitSystem_SessionDidUpdateAnchors, STATGROUP_ARKIT);
+DECLARE_CYCLE_STAT(TEXT("SessionDidRemoveAnchors_DelegateThread"), STAT_FAppleARKitSystem_SessionDidRemoveAnchors, STATGROUP_ARKIT);
+DECLARE_CYCLE_STAT(TEXT("UpdateARKitPerf"), STAT_FAppleARKitSystem_UpdateARKitPerf, STATGROUP_ARKIT);
+DECLARE_DWORD_COUNTER_STAT(TEXT("ARKit CPU %"), STAT_ARKitThreads, STATGROUP_ARKIT);
+
 // Copied from IOSPlatformProcess because it's not accessible by external code
 #define GAME_THREAD_PRIORITY 47
 #define RENDER_THREAD_PRIORITY 45
@@ -1362,13 +1369,13 @@ void FAppleARKitSystem::OrientationChanged(const int32 NewOrientationRaw)
 						
 void FAppleARKitSystem::SessionDidUpdateFrame_DelegateThread(TSharedPtr< FAppleARKitFrame, ESPMode::ThreadSafe > Frame)
 {
-	// Thread safe swap buffered frame
-	DECLARE_CYCLE_STAT(TEXT("FAppleARKitSystem::SessionDidUpdateFrame_DelegateThread"),
-					   STAT_FAppleARKitSystem_SessionUpdateFrame,
-					   STATGROUP_APPLEARKIT);
-	
-	auto UpdateFrameTask = FSimpleDelegateGraphTask::FDelegate::CreateThreadSafeSP( this, &FAppleARKitSystem::SessionDidUpdateFrame_Internal, Frame.ToSharedRef() );
-	FSimpleDelegateGraphTask::CreateAndDispatchWhenReady(UpdateFrameTask, GET_STATID(STAT_FAppleARKitSystem_SessionUpdateFrame), nullptr, ENamedThreads::GameThread);
+	{
+		auto UpdateFrameTask = FSimpleDelegateGraphTask::FDelegate::CreateThreadSafeSP( this, &FAppleARKitSystem::SessionDidUpdateFrame_Internal, Frame.ToSharedRef() );
+		FSimpleDelegateGraphTask::CreateAndDispatchWhenReady(UpdateFrameTask, GET_STATID(STAT_FAppleARKitSystem_SessionUpdateFrame), nullptr, ENamedThreads::GameThread);
+	}
+	{
+		UpdateARKitPerfStats();
+	}
 }
 			
 void FAppleARKitSystem::SessionDidFailWithError_DelegateThread(const FString& Error)
@@ -1457,10 +1464,6 @@ static TSharedPtr<FAppleARKitAnchorData> MakeAnchorData( ARAnchor* Anchor )
 
 void FAppleARKitSystem::SessionDidAddAnchors_DelegateThread( NSArray<ARAnchor*>* anchors )
 {
-	DECLARE_CYCLE_STAT(TEXT("FAppleARKitSystem::SessionDidAddAnchors_DelegateThread"),
-					   STAT_FAppleARKitSystem_SessionDidAddAnchors,
-					   STATGROUP_APPLEARKIT);
-
 	// If this object is valid, we are running a face session and need that code to process things
 	if (FaceARSupport != nullptr)
 	{
@@ -1490,10 +1493,6 @@ void FAppleARKitSystem::SessionDidAddAnchors_DelegateThread( NSArray<ARAnchor*>*
 
 void FAppleARKitSystem::SessionDidUpdateAnchors_DelegateThread( NSArray<ARAnchor*>* anchors )
 {
-	DECLARE_CYCLE_STAT(TEXT("FAppleARKitSystem::SessionDidUpdateAnchors_DelegateThread"),
-					   STAT_FAppleARKitSystem_SessionDidUpdateAnchors,
-					   STATGROUP_APPLEARKIT);
-	
 	// If this object is valid, we are running a face session and need that code to process things
 	if (FaceARSupport != nullptr)
 	{
@@ -1523,12 +1522,7 @@ void FAppleARKitSystem::SessionDidUpdateAnchors_DelegateThread( NSArray<ARAnchor
 
 void FAppleARKitSystem::SessionDidRemoveAnchors_DelegateThread( NSArray<ARAnchor*>* anchors )
 {
-	DECLARE_CYCLE_STAT(TEXT("FAppleARKitSystem::SessionDidRemoveAnchors_DelegateThread"),
-					   STAT_FAppleARKitSystem_SessionDidRemoveAnchors,
-					   STATGROUP_APPLEARKIT);
-	
 	// Face AR Anchors are also removed this way, no need for special code since they are tracked geometry
-
 	for (ARAnchor* anchor in anchors)
 	{
 		// Convert to FGuid
@@ -1768,8 +1762,96 @@ void FAppleARKitSystem::SessionDidUpdateFrame_Internal( TSharedRef< FAppleARKitF
 	UpdateFrame();
 }
 
+#if STATS
+struct FARKitThreadTimes
+{
+	TArray<FString> ThreadNames;
+	int32 LastTotal;
+	int32 NewTotal;
+	
+	FARKitThreadTimes() :
+		LastTotal(0)
+		, NewTotal(0)
+	{
+		ThreadNames.Add(TEXT("com.apple.CoreMotion"));
+		ThreadNames.Add(TEXT("com.apple.arkit"));
+		ThreadNames.Add(TEXT("FilteringFrameDownsampleNodeWorkQueue"));
+		ThreadNames.Add(TEXT("FeatureDetectorNodeWorkQueue"));
+		ThreadNames.Add(TEXT("SurfaceDetectionNode"));
+		ThreadNames.Add(TEXT("VIOEngineNode"));
+		ThreadNames.Add(TEXT("ImageDetectionQueue"));
+	}
 
+	bool IsARKitThread(const FString& Name)
+	{
+		if (Name.Len() == 0)
+		{
+			return false;
+		}
+		
+		for (int32 Index = 0; Index < ThreadNames.Num(); Index++)
+		{
+			if (Name.StartsWith(ThreadNames[Index]))
+			{
+				return true;
+			}
+		}
+		return false;
+	}
+	
+	void FrameReset()
+	{
+		LastTotal = NewTotal;
+		NewTotal = 0;
+	}
+};
+#endif
 
+void FAppleARKitSystem::UpdateARKitPerfStats()
+{
+#if STATS && SUPPORTS_ARKIT_1_0
+	static FARKitThreadTimes ARKitThreadTimes;
+
+	SCOPE_CYCLE_COUNTER(STAT_FAppleARKitSystem_UpdateARKitPerf);
+	ARKitThreadTimes.FrameReset();
+	
+	thread_array_t ThreadArray;
+	mach_msg_type_number_t ThreadCount;
+	if (task_threads(mach_task_self(), &ThreadArray, &ThreadCount) != KERN_SUCCESS)
+	{
+		return;
+	}
+
+	for (int32 Index = 0; Index < (int32)ThreadCount; Index++)
+	{
+		mach_msg_type_number_t ThreadInfoCount = THREAD_BASIC_INFO_COUNT;
+		mach_msg_type_number_t ExtThreadInfoCount = THREAD_EXTENDED_INFO_COUNT;
+		thread_info_data_t ThreadInfo;
+		thread_extended_info_data_t ExtThreadInfo;
+		// Get the basic thread info for this thread
+		if (thread_info(ThreadArray[Index], THREAD_BASIC_INFO, (thread_info_t)ThreadInfo, &ThreadInfoCount) != KERN_SUCCESS)
+		{
+			continue;
+		}
+		// And the extended thread info for this thread
+		if (thread_info(ThreadArray[Index], THREAD_EXTENDED_INFO, (thread_info_t)&ExtThreadInfo, &ExtThreadInfoCount) != KERN_SUCCESS)
+		{
+			continue;
+		}
+		thread_basic_info_t BasicInfo = (thread_basic_info_t)ThreadInfo;
+		FString ThreadName(ExtThreadInfo.pth_name);
+		if (ARKitThreadTimes.IsARKitThread(ThreadName))
+		{
+			// CPU usage is reported as a scaled number, so convert to %
+			int32 ScaledPercent = FMath::RoundToInt((float)BasicInfo->cpu_usage / (float)TH_USAGE_SCALE * 100.f);
+			ARKitThreadTimes.NewTotal += ScaledPercent;
+		}
+//		UE_LOG(LogAppleARKit, Log, TEXT("Thread %s used cpu (%d), seconds (%d), microseconds (%d)"), *ThreadName, BasicInfo->cpu_usage, BasicInfo->user_time.seconds + BasicInfo->system_time.seconds, BasicInfo->user_time.microseconds + BasicInfo->system_time.microseconds);
+	}
+	vm_deallocate(mach_task_self(), (vm_offset_t)ThreadArray, ThreadCount * sizeof(thread_t));
+	SET_DWORD_STAT(STAT_ARKitThreads, ARKitThreadTimes.NewTotal);
+#endif
+}
 
 
 namespace AppleARKitSupport
diff --git a/Engine/Plugins/Runtime/AppleARKit/Source/AppleARKit/Private/AppleARKitSystem.h b/Engine/Plugins/Runtime/AppleARKit/Source/AppleARKit/Private/AppleARKitSystem.h
--- a/Engine/Plugins/Runtime/AppleARKit/Source/AppleARKit/Private/AppleARKitSystem.h
+++ b/Engine/Plugins/Runtime/AppleARKit/Source/AppleARKit/Private/AppleARKitSystem.h
@@ -16,7 +16,7 @@
 #endif
 
 
-DECLARE_STATS_GROUP(TEXT("AppleARKit"), STATGROUP_APPLEARKIT, STATCAT_Advanced);
+DECLARE_STATS_GROUP(TEXT("ARKit"), STATGROUP_ARKIT, STATCAT_Advanced);
 
 //
 //  FAppleARKitSystem
@@ -137,6 +137,9 @@ private:
 	/** Creates or clears the face ar support object if face ar has been requested */
 	void CheckForFaceARSupport(UARSessionConfig* InSessionConfig);
 	
+	/** Updates the ARKit perf counters */
+	void UpdateARKitPerfStats();
+	
 	/** The orientation of the device; see EScreenOrientation */
 	EScreenOrientation::Type DeviceOrientation;
 	
