Merging using Dev-VR_to_Release-4.21

diff --git a/Engine/Plugins/Runtime/AppleARKit/Source/AppleARKit/Public/AppleARKitFaceSupport.h b/Engine/Plugins/Runtime/AppleARKit/Source/AppleARKit/Public/AppleARKitFaceSupport.h
--- a/Engine/Plugins/Runtime/AppleARKit/Source/AppleARKit/Public/AppleARKitFaceSupport.h
+++ b/Engine/Plugins/Runtime/AppleARKit/Source/AppleARKit/Public/AppleARKitFaceSupport.h
@@ -11,7 +11,8 @@
 	#import <ARKit/ARKit.h>
 #endif
 
-class UARSessionConfig;
+#include "ARSessionConfig.h"
+
 struct FAppleARKitAnchorData;
 struct FARVideoFormat;
 
