Instead of an assert (and thus crash) if Session Configuration is null (which happens if device/os doesn't support ARKit 2.0), just print out an error to the log and keep running.

diff --git a/Engine/Plugins/Runtime/AppleARKit/Source/AppleARKit/Private/AppleARKitConversion.cpp b/Engine/Plugins/Runtime/AppleARKit/Source/AppleARKit/Private/AppleARKitConversion.cpp
--- a/Engine/Plugins/Runtime/AppleARKit/Source/AppleARKit/Private/AppleARKitConversion.cpp
+++ b/Engine/Plugins/Runtime/AppleARKit/Source/AppleARKit/Private/AppleARKitConversion.cpp
@@ -329,14 +329,19 @@ ARConfiguration* FAppleARKitConversion::ToARConfiguration( UARSessionConfig* Ses
 		default:
 			return nullptr;
 	}
-	check(SessionConfiguration != nullptr);
-
-	// Copy / convert properties
-	SessionConfiguration.lightEstimationEnabled = SessionConfig->GetLightEstimationMode() != EARLightEstimationMode::None;
-	SessionConfiguration.providesAudioData = NO;
-	SessionConfiguration.worldAlignment = FAppleARKitConversion::ToARWorldAlignment(SessionConfig->GetWorldAlignment());
-
-	return SessionConfiguration;
+		if (SessionConfiguration != nullptr)
+        {
+            // Copy / convert properties
+            SessionConfiguration.lightEstimationEnabled = SessionConfig->GetLightEstimationMode() != EARLightEstimationMode::None;
+            SessionConfiguration.providesAudioData = NO;
+            SessionConfiguration.worldAlignment = FAppleARKitConversion::ToARWorldAlignment(SessionConfig->GetWorldAlignment());
+        }
+        else
+        {
+            UE_LOG(LogAppleARKit, Error, TEXT("Failed to create the requested ARKit Session Type. Please check that the device supports this session type."));
+        }
+    
+    return SessionConfiguration;
 }
 
 #pragma clang diagnostic pop
