Removing error logging since it's also getting logged higher up in the callstack.

diff --git a/Engine/Plugins/Runtime/AppleARKit/Source/AppleARKit/Private/AppleARKitConversion.cpp b/Engine/Plugins/Runtime/AppleARKit/Source/AppleARKit/Private/AppleARKitConversion.cpp
--- a/Engine/Plugins/Runtime/AppleARKit/Source/AppleARKit/Private/AppleARKitConversion.cpp
+++ b/Engine/Plugins/Runtime/AppleARKit/Source/AppleARKit/Private/AppleARKitConversion.cpp
@@ -329,17 +329,13 @@ ARConfiguration* FAppleARKitConversion::ToARConfiguration( UARSessionConfig* Ses
 		default:
 			return nullptr;
 	}
-		if (SessionConfiguration != nullptr)
-        {
-            // Copy / convert properties
-            SessionConfiguration.lightEstimationEnabled = SessionConfig->GetLightEstimationMode() != EARLightEstimationMode::None;
-            SessionConfiguration.providesAudioData = NO;
-            SessionConfiguration.worldAlignment = FAppleARKitConversion::ToARWorldAlignment(SessionConfig->GetWorldAlignment());
-        }
-        else
-        {
-            UE_LOG(LogAppleARKit, Error, TEXT("Failed to create the requested ARKit Session Type. Please check that the device supports this session type."));
-        }
+    if (SessionConfiguration != nullptr)
+    {
+        // Copy / convert properties
+        SessionConfiguration.lightEstimationEnabled = SessionConfig->GetLightEstimationMode() != EARLightEstimationMode::None;
+        SessionConfiguration.providesAudioData = NO;
+        SessionConfiguration.worldAlignment = FAppleARKitConversion::ToARWorldAlignment(SessionConfig->GetWorldAlignment());
+    }
     
     return SessionConfiguration;
 }
