Fixed the use of deprecated NSKeyedArchiver and NSKeyedUnarchiver methods

diff --git a/Engine/Plugins/Runtime/AppleARKit/Source/AppleARKit/Private/AppleARKitConversion.cpp b/Engine/Plugins/Runtime/AppleARKit/Source/AppleARKit/Private/AppleARKitConversion.cpp
--- a/Engine/Plugins/Runtime/AppleARKit/Source/AppleARKit/Private/AppleARKitConversion.cpp
+++ b/Engine/Plugins/Runtime/AppleARKit/Source/AppleARKit/Private/AppleARKitConversion.cpp
@@ -176,10 +176,12 @@ ARWorldMap* FAppleARKitConversion::ToARWorldMap(const TArray<uint8>& WorldMapDat
 	
 	// Serialize into the World map data
 	NSData* WorldNSData = [NSData dataWithBytesNoCopy: UncompressedData.GetData() length: UncompressedData.Num() freeWhenDone: NO];
-	ARWorldMap* WorldMap = [NSKeyedUnarchiver unarchiveObjectWithData: WorldNSData];
-	if (WorldMap == nullptr)
+	NSError* ErrorObj = nullptr;
+	ARWorldMap* WorldMap = [NSKeyedUnarchiver unarchivedObjectOfClass: ARWorldMap.class fromData: WorldNSData error: &ErrorObj];
+	if (ErrorObj != nullptr)
 	{
-		UE_LOG(LogAppleARKit, Log, TEXT("Failed to load the world map data from the session object"));
+		FString Error = [ErrorObj localizedDescription];
+		UE_LOG(LogAppleARKit, Log, TEXT("Failed to load the world map data from the session object with error string (%s)"), *Error);
 	}
 	return WorldMap;
 }
@@ -199,7 +201,8 @@ NSSet* FAppleARKitConversion::ToARReferenceObjectSet(const TArray<UARCandidateOb
 		if (Candidate != nullptr && Candidate->GetCandidateObjectData().Num() > 0)
 		{
 			NSData* CandidateData = [NSData dataWithBytesNoCopy: (uint8*)Candidate->GetCandidateObjectData().GetData() length: Candidate->GetCandidateObjectData().Num() freeWhenDone: NO];
-			ARReferenceObject* RefObject = [NSKeyedUnarchiver unarchiveObjectWithData: CandidateData];
+			NSError* ErrorObj = nullptr;
+			ARReferenceObject* RefObject = [NSKeyedUnarchiver unarchivedObjectOfClass: ARReferenceObject.class fromData: CandidateData error: &ErrorObj];
 			if (RefObject != nullptr)
 			{
 				// Store off so the session object can quickly match the anchor to our representation
diff --git a/Engine/Plugins/Runtime/AppleARKit/Source/AppleARKit/Private/AppleARKitSystem.cpp b/Engine/Plugins/Runtime/AppleARKit/Source/AppleARKit/Private/AppleARKitSystem.cpp
--- a/Engine/Plugins/Runtime/AppleARKit/Source/AppleARKit/Private/AppleARKitSystem.cpp
+++ b/Engine/Plugins/Runtime/AppleARKit/Source/AppleARKit/Private/AppleARKitSystem.cpp
@@ -950,7 +950,8 @@ public:
 			CandidateObject->SetBoundingBox(BoundingBox);
 			
 			// Serialize the object into a byte array and stick that on the candidate object
-			NSData* RefObjData = [NSKeyedArchiver archivedDataWithRootObject: ReferenceObject];
+			NSError* ErrorObj = nullptr;
+			NSData* RefObjData = [NSKeyedArchiver archivedDataWithRootObject: ReferenceObject requiringSecureCoding: YES error: &ErrorObj];
 			uint32 SavedSize = RefObjData.length;
 			TArray<uint8> RawBytes;
 			RawBytes.AddUninitialized(SavedSize);
@@ -1020,28 +1021,36 @@ public:
 	{
 		if (bWasSuccessful)
 		{
-			NSData* WorldNSData = [NSKeyedArchiver archivedDataWithRootObject: WorldMap];
-
-			int32 UncompressedSize = WorldNSData.length;
-
-			TArray<uint8> CompressedData;
-			CompressedData.AddUninitialized(WorldNSData.length + AR_SAVE_WORLD_HEADER_SIZE);
-			uint8* Buffer = (uint8*)CompressedData.GetData();
-			// Write our magic header into our buffer
-			FARWorldSaveHeader& Header = *(FARWorldSaveHeader*)Buffer;
-			Header = FARWorldSaveHeader();
-			Header.UncompressedSize = UncompressedSize;
-			
-			// Compress the data
-			uint8* CompressInto = Buffer + AR_SAVE_WORLD_HEADER_SIZE;
-			int32 CompressedSize = UncompressedSize;
-			uint8* UncompressedData = (uint8*)[WorldNSData bytes];
-			verify(FCompression::CompressMemory((ECompressionFlags)COMPRESS_ZLIB, CompressInto, CompressedSize, UncompressedData, UncompressedSize));
-			
-			// Only copy out the amount of compressed data and the header
-			int32 CompressedSizePlusHeader = CompressedSize + AR_SAVE_WORLD_HEADER_SIZE;
-			WorldData.AddUninitialized(CompressedSizePlusHeader);
-			FPlatformMemory::Memcpy(WorldData.GetData(), CompressedData.GetData(), CompressedSizePlusHeader);
+			NSError* ErrorObj = nullptr;
+			NSData* WorldNSData = [NSKeyedArchiver archivedDataWithRootObject: WorldMap requiringSecureCoding: YES error: &ErrorObj];
+			if (ErrorObj == nullptr)
+			{
+				int32 UncompressedSize = WorldNSData.length;
+				
+				TArray<uint8> CompressedData;
+				CompressedData.AddUninitialized(WorldNSData.length + AR_SAVE_WORLD_HEADER_SIZE);
+				uint8* Buffer = (uint8*)CompressedData.GetData();
+				// Write our magic header into our buffer
+				FARWorldSaveHeader& Header = *(FARWorldSaveHeader*)Buffer;
+				Header = FARWorldSaveHeader();
+				Header.UncompressedSize = UncompressedSize;
+				
+				// Compress the data
+				uint8* CompressInto = Buffer + AR_SAVE_WORLD_HEADER_SIZE;
+				int32 CompressedSize = UncompressedSize;
+				uint8* UncompressedData = (uint8*)[WorldNSData bytes];
+				verify(FCompression::CompressMemory((ECompressionFlags)COMPRESS_ZLIB, CompressInto, CompressedSize, UncompressedData, UncompressedSize));
+				
+				// Only copy out the amount of compressed data and the header
+				int32 CompressedSizePlusHeader = CompressedSize + AR_SAVE_WORLD_HEADER_SIZE;
+				WorldData.AddUninitialized(CompressedSizePlusHeader);
+				FPlatformMemory::Memcpy(WorldData.GetData(), CompressedData.GetData(), CompressedSizePlusHeader);
+			}
+			else
+			{
+				Error = [ErrorObj localizedDescription];
+				bHadError = true;
+			}
 		}
 		else
 		{
