Changed the world mapping state conversion in ARKit to return NotAvailable, Mapping, Mapped

diff --git a/Engine/Plugins/Runtime/AppleARKit/Source/AppleARKit/Private/AppleARKitFrame.cpp b/Engine/Plugins/Runtime/AppleARKit/Source/AppleARKit/Private/AppleARKitFrame.cpp
--- a/Engine/Plugins/Runtime/AppleARKit/Source/AppleARKit/Private/AppleARKitFrame.cpp
+++ b/Engine/Plugins/Runtime/AppleARKit/Source/AppleARKit/Private/AppleARKitFrame.cpp
@@ -28,11 +28,10 @@ EARWorldMappingState ToEARWorldMappingState(ARWorldMappingStatus MapStatus)
 {
 	switch (MapStatus)
 	{
+		// These both mean more data is needed
 		case ARWorldMappingStatusLimited:
-			return EARWorldMappingState::StillMappingNotRelocalizable;
-
 		case ARWorldMappingStatusExtending:
-			return EARWorldMappingState::StillMappingRelocalizable;
+			return EARWorldMappingState::StillMappingNotRelocalizable;
 
 		case ARWorldMappingStatusMapped:
 			return EARWorldMappingState::Mapped;
