Added a check for an invalid physical size when specifying a candidate image for detection within a AR scene

diff --git a/Engine/Plugins/Runtime/AppleARKit/Source/AppleARKit/Private/AppleARKitConversion.cpp b/Engine/Plugins/Runtime/AppleARKit/Source/AppleARKit/Private/AppleARKitConversion.cpp
--- a/Engine/Plugins/Runtime/AppleARKit/Source/AppleARKit/Private/AppleARKitConversion.cpp
+++ b/Engine/Plugins/Runtime/AppleARKit/Source/AppleARKit/Private/AppleARKitConversion.cpp
@@ -83,6 +83,13 @@ NSSet* FAppleARKitConversion::InitImageDetection(UARSessionConfig* SessionConfig
 	{
 		if (Candidate != nullptr && Candidate->GetCandidateTexture() != nullptr)
 		{
+			// Don't crash if the physical size is invalid
+			if (Candidate->GetPhysicalWidth() <= 0.f || Candidate->GetPhysicalHeight() <= 0.f)
+			{
+				UE_LOG(LogAppleARKit, Error, TEXT("Unable to process candidate image (%s - %s) due to an invalid physical size (%f,%f)"),
+				   *Candidate->GetFriendlyName(), *Candidate->GetName(), Candidate->GetPhysicalWidth(), Candidate->GetPhysicalHeight());
+				continue;
+			}
 			// Store off so the session object can quickly match the anchor to our representation
 			// This stores it even if we weren't able to convert to apple's type for GC reasons
 			CandidateImages.Add(Candidate->GetFriendlyName(), Candidate);
