UE4 - IOS Compile fixes fallout from integration

diff --git a/Engine/Plugins/Runtime/AppleARKit/Source/AppleARKit/Private/AppleARKitVideoOverlay.cpp b/Engine/Plugins/Runtime/AppleARKit/Source/AppleARKit/Private/AppleARKitVideoOverlay.cpp
--- a/Engine/Plugins/Runtime/AppleARKit/Source/AppleARKit/Private/AppleARKitVideoOverlay.cpp
+++ b/Engine/Plugins/Runtime/AppleARKit/Source/AppleARKit/Private/AppleARKitVideoOverlay.cpp
@@ -395,7 +395,7 @@ void FAppleARKitVideoOverlay::RenderVideoOverlay_RenderThread(FRHICommandListImm
 		GraphicsPSOInit.RasterizerState = TStaticRasterizerState<>::GetRHI();
 		GraphicsPSOInit.DepthStencilState = TStaticDepthStencilState<false, CF_DepthNearOrEqual>::GetRHI();
 		GraphicsPSOInit.PrimitiveType = PT_TriangleList;
-		GraphicsPSOInit.BoundShaderState.VertexDeclarationRHI = GFilterVertexDeclaration.VertexDeclarationRH;
+		GraphicsPSOInit.BoundShaderState.VertexDeclarationRHI = GFilterVertexDeclaration.VertexDeclarationRHI;
 		
 		const bool bIsMobileRenderer = FeatureLevel <= ERHIFeatureLevel::ES3_1;
 		FMaterialShader* VertexShader = nullptr;
@@ -469,7 +469,6 @@ void FAppleARKitVideoOverlay::RenderVideoOverlay_RenderThread(FRHICommandListImm
 			RHICmdList.SetStreamSource(0, VertexBufferRHI, 0);
 			RHICmdList.DrawIndexedPrimitive(
 				OverlayIndexBufferRHI,
-				PT_TriangleList,
 				/*BaseVertexIndex=*/ 0,
 				/*MinIndex=*/ 0,
 				/*NumVertices=*/ 4,
