Copying from Release-4.21: 4611082 UE-63052 Handheld AR camera is flipped on iPhone 8+
-The PreferedLandscapeOrientation feature uncovered flaws in how we are setting the ARKit orientation so that it matches the data coming in from ARKit.
-We now use statusBarOrientation on startup to detect the orientation of the app.  This is better than device orientation which is a more raw view of what device sensors are saying.
-I have updated some of the code that handles orientation change events to account for the preferredlandscapeorientation so that it will chose left or right landscape appropriately.  It would be better to change this code to handling the statsBarOrientationChanged event, but we can't do that in a release patch.  I did not see orientation work incorrectly or flicker.

diff --git a/Engine/Plugins/Runtime/AppleARKit/Source/AppleARKit/Private/AppleARKitSystem.cpp b/Engine/Plugins/Runtime/AppleARKit/Source/AppleARKit/Private/AppleARKitSystem.cpp
--- a/Engine/Plugins/Runtime/AppleARKit/Source/AppleARKit/Private/AppleARKitSystem.cpp
+++ b/Engine/Plugins/Runtime/AppleARKit/Source/AppleARKit/Private/AppleARKitSystem.cpp
@@ -1167,25 +1167,51 @@ void FAppleARKitSystem::SetDeviceOrientation( EScreenOrientation::Type InOrienta
 		// So pick ANY ALLOWED default.
 		// This only realy happens if the device is face down on something or
 		// in another "useless" state for AR.
+
+		// Note: the order in which this selection is done is important and must match that 
+		// established in UEDeployIOS.cs and written into UISupportedInterfaceOrientations.
+		// IOSView preferredInterfaceOrientationForPresentation presumably also should match.
+		//
+		// However it would very likely be better to hook statusBarOrientation instead of deviceOrientation to update
+		// our orientation, in which case we would only need to handle unknown.
 		
 		if (!NewOrientation.IsSet())
 		{
 			NewOrientation = PickAllowedDeviceOrientation(EScreenOrientation::Portrait);
 		}
-		
+
 		if (!NewOrientation.IsSet())
 		{
-			NewOrientation = PickAllowedDeviceOrientation(EScreenOrientation::LandscapeLeft);
+			NewOrientation = PickAllowedDeviceOrientation(EScreenOrientation::PortraitUpsideDown);
 		}
-		
-		if (!NewOrientation.IsSet())
+
+#if SUPPORTS_ARKIT_1_0
+		const UIOSRuntimeSettings* IOSSettings = GetDefault<UIOSRuntimeSettings>();
+		const bool bPreferLandscapeLeftHomeButton = IOSSettings->PreferredLandscapeOrientation == EIOSLandscapeOrientation::LandscapeLeft;
+#else
+		const bool bPreferLandscapeLeftHomeButton = true;
+#endif
+		if (bPreferLandscapeLeftHomeButton)
 		{
-			NewOrientation = PickAllowedDeviceOrientation(EScreenOrientation::PortraitUpsideDown);
+			if (!NewOrientation.IsSet())
+			{
+				NewOrientation = PickAllowedDeviceOrientation(EScreenOrientation::LandscapeRight);
+			}
+			if (!NewOrientation.IsSet())
+			{
+				NewOrientation = PickAllowedDeviceOrientation(EScreenOrientation::LandscapeLeft);
+			}
 		}
-		
-		if (!NewOrientation.IsSet())
+		else
 		{
-			NewOrientation = PickAllowedDeviceOrientation(EScreenOrientation::LandscapeRight);
+			if (!NewOrientation.IsSet())
+			{
+				NewOrientation = PickAllowedDeviceOrientation(EScreenOrientation::LandscapeLeft);
+			}
+			if (!NewOrientation.IsSet())
+			{
+				NewOrientation = PickAllowedDeviceOrientation(EScreenOrientation::LandscapeRight);
+			}
 		}
 		
 		check(NewOrientation.IsSet());
@@ -1198,6 +1224,28 @@ void FAppleARKitSystem::SetDeviceOrientation( EScreenOrientation::Type InOrienta
 	}
 }
 
+static EScreenOrientation::Type GetAppOrientation()
+{
+#if PLATFORM_IOS && !PLATFORM_TVOS
+	// We want the orientation that the app is running with, not necessarily the orientation of the device right now.
+	UIInterfaceOrientation Orientation = [[UIApplication sharedApplication] statusBarOrientation];
+#if __IPHONE_OS_VERSION_MIN_REQUIRED < __IPHONE_8_0
+	Orientation = [[IOSAppDelegate GetDelegate].IOSController interfaceOrientation];
+#endif
+	EScreenOrientation::Type ScreenOrientation = EScreenOrientation::Unknown;
+	switch (Orientation)
+	{
+	case UIInterfaceOrientationUnknown:				return EScreenOrientation::Unknown;
+	case UIInterfaceOrientationPortrait:			return EScreenOrientation::Portrait;
+	case UIInterfaceOrientationPortraitUpsideDown:	return EScreenOrientation::PortraitUpsideDown;
+	case UIInterfaceOrientationLandscapeLeft:		return EScreenOrientation::LandscapeRight;
+	case UIInterfaceOrientationLandscapeRight:		return EScreenOrientation::LandscapeLeft;
+	default:										check(false); return EScreenOrientation::Unknown;
+	}
+#else
+	return static_cast<EScreenOrientation::Type>(FPlatformMisc::GetDeviceOrientation());
+#endif
+}
 
 PRAGMA_DISABLE_OPTIMIZATION
 bool FAppleARKitSystem::Run(UARSessionConfig* SessionConfig)
@@ -1212,7 +1260,8 @@ bool FAppleARKitSystem::Run(UARSessionConfig* SessionConfig)
 	// Make sure this is set at session start, because there are timing issues with using only the delegate approach
 	if (DeviceOrientation == EScreenOrientation::Unknown)
 	{
-		SetDeviceOrientation( static_cast<EScreenOrientation::Type>(FPlatformMisc::GetDeviceOrientation()) );
+		const EScreenOrientation::Type ScreenOrientation = GetAppOrientation();
+		SetDeviceOrientation( ScreenOrientation );
 	}
 
 #if SUPPORTS_ARKIT_1_0
