Fixed AppleARKit not shutting down properly when the application is terminated by the user on a device

diff --git a/Engine/Plugins/Runtime/AppleARKit/Source/AppleARKit/Private/AppleARKitModule.cpp b/Engine/Plugins/Runtime/AppleARKit/Source/AppleARKit/Private/AppleARKitModule.cpp
--- a/Engine/Plugins/Runtime/AppleARKit/Source/AppleARKit/Private/AppleARKitModule.cpp
+++ b/Engine/Plugins/Runtime/AppleARKit/Source/AppleARKit/Private/AppleARKitModule.cpp
@@ -6,6 +6,7 @@
 #include "Features/IModularFeature.h"
 #include "Features/IModularFeatures.h"
 #include "CoreGlobals.h"
+#include "Misc/CoreDelegates.h"
 
 
 TWeakPtr<class FAppleARKitSystem, ESPMode::ThreadSafe> FAppleARKitARKitSystemPtr;
@@ -37,6 +38,8 @@ void FAppleARKitModule::StartupModule()
 	ensureMsgf(FModuleManager::Get().LoadModule("AugmentedReality"), TEXT("ARKit depends on the AugmentedReality module."));
 	IHeadMountedDisplayModule::StartupModule();
 
+	FCoreDelegates::OnPreExit.AddRaw(this, &FAppleARKitModule::PreExit);
+
 	// LiveLink listener needs to be created here so that the editor can receive remote publishing events
 #if PLATFORM_DESKTOP
 	bool bEnableLiveLinkForFaceTracking = false;
@@ -48,6 +51,15 @@ void FAppleARKitModule::StartupModule()
 #endif
 }
 
+void FAppleARKitModule::PreExit()
+{
+	if (FAppleARKitARKitSystemPtr.IsValid())
+	{
+		FAppleARKitARKitSystemPtr.Pin()->Shutdown();
+	}
+	FAppleARKitARKitSystemPtr.Reset();
+}
+
 void FAppleARKitModule::ShutdownModule()
 {
 	IHeadMountedDisplayModule::ShutdownModule();
diff --git a/Engine/Plugins/Runtime/AppleARKit/Source/AppleARKit/Private/AppleARKitModule.h b/Engine/Plugins/Runtime/AppleARKit/Source/AppleARKit/Private/AppleARKitModule.h
--- a/Engine/Plugins/Runtime/AppleARKit/Source/AppleARKit/Private/AppleARKitModule.h
+++ b/Engine/Plugins/Runtime/AppleARKit/Source/AppleARKit/Private/AppleARKitModule.h
@@ -21,6 +21,7 @@ public:
 
 	virtual void ShutdownModule() override;
 
+    void PreExit();
 };
 
 DECLARE_LOG_CATEGORY_EXTERN(LogAppleARKit, Log, All);
diff --git a/Engine/Plugins/Runtime/AppleARKit/Source/AppleARKit/Private/AppleARKitSystem.cpp b/Engine/Plugins/Runtime/AppleARKit/Source/AppleARKit/Private/AppleARKitSystem.cpp
--- a/Engine/Plugins/Runtime/AppleARKit/Source/AppleARKit/Private/AppleARKitSystem.cpp
+++ b/Engine/Plugins/Runtime/AppleARKit/Source/AppleARKit/Private/AppleARKitSystem.cpp
@@ -429,6 +429,20 @@ FAppleARKitSystem::~FAppleARKitSystem()
 	// Unregister our ability to hit-test in AR with Unreal
 }
 
+void FAppleARKitSystem::Shutdown()
+{
+#if SUPPORTS_ARKIT_1_0
+	if (Session != nullptr)
+	{
+		[Session pause];
+		Session.delegate = nullptr;
+		[Session release];
+		Session = nullptr;
+	}
+#endif
+	CameraDepth = nullptr;
+	CameraImage = nullptr;
+}
 
 FName FAppleARKitSystem::GetSystemName() const
 {
diff --git a/Engine/Plugins/Runtime/AppleARKit/Source/AppleARKit/Private/AppleARKitSystem.h b/Engine/Plugins/Runtime/AppleARKit/Source/AppleARKit/Private/AppleARKitSystem.h
--- a/Engine/Plugins/Runtime/AppleARKit/Source/AppleARKit/Private/AppleARKitSystem.h
+++ b/Engine/Plugins/Runtime/AppleARKit/Source/AppleARKit/Private/AppleARKitSystem.h
@@ -53,6 +53,10 @@ public:
 	
 	// @todo arkit : this is for the blueprint library only; try to get rid of this method
 	bool GetCurrentFrame(FAppleARKitFrame& OutCurrentFrame) const;
+
+	/** So the module can shut down the ar services cleanly */
+	void Shutdown();
+
 private:
 	//~ FGCObject
 	virtual void AddReferencedObjects( FReferenceCollector& Collector ) override;
