Prep work for moving face ar to its own plugin

diff --git a/Engine/Plugins/Runtime/AppleARKit/Source/AppleARKit/Private/AppleARKitAnchor.cpp b/Engine/Plugins/Runtime/AppleARKit/Source/AppleARKit/Private/AppleARKitAnchor.cpp
--- a/Engine/Plugins/Runtime/AppleARKit/Source/AppleARKit/Private/AppleARKitAnchor.cpp
+++ b/Engine/Plugins/Runtime/AppleARKit/Source/AppleARKit/Private/AppleARKitAnchor.cpp
@@ -3,7 +3,7 @@
 // AppleARKit
 #include "AppleARKitAnchor.h"
 #include "AppleARKitModule.h"
-#include "AppleARKitTransform.h"
+#include "AppleARKitConversion.h"
 
 // UE4
 #include "Misc/ScopeLock.h"
diff --git a/Engine/Plugins/Runtime/AppleARKit/Source/AppleARKit/Private/AppleARKitCamera.cpp b/Engine/Plugins/Runtime/AppleARKit/Source/AppleARKit/Private/AppleARKitCamera.cpp
--- a/Engine/Plugins/Runtime/AppleARKit/Source/AppleARKit/Private/AppleARKitCamera.cpp
+++ b/Engine/Plugins/Runtime/AppleARKit/Source/AppleARKit/Private/AppleARKitCamera.cpp
@@ -3,7 +3,7 @@
 // AppleARKit
 #include "AppleARKitCamera.h"
 #include "AppleARKitModule.h"
-#include "AppleARKitTransform.h"
+#include "AppleARKitConversion.h"
 	
 // UE4
 #include "CoreGlobals.h"
diff --git a/Engine/Plugins/Runtime/AppleARKit/Source/AppleARKit/Private/AppleARKitHitTestResult.cpp b/Engine/Plugins/Runtime/AppleARKit/Source/AppleARKit/Private/AppleARKitHitTestResult.cpp
--- a/Engine/Plugins/Runtime/AppleARKit/Source/AppleARKit/Private/AppleARKitHitTestResult.cpp
+++ b/Engine/Plugins/Runtime/AppleARKit/Source/AppleARKit/Private/AppleARKitHitTestResult.cpp
@@ -2,7 +2,7 @@
 
 // AppleARKit
 #include "AppleARKitHitTestResult.h"
-#include "AppleARKitTransform.h"
+#include "AppleARKitConversion.h"
 
 #if SUPPORTS_ARKIT_1_0
 
diff --git a/Engine/Plugins/Runtime/AppleARKit/Source/AppleARKit/Private/AppleARKitPlaneAnchor.cpp b/Engine/Plugins/Runtime/AppleARKit/Source/AppleARKit/Private/AppleARKitPlaneAnchor.cpp
--- a/Engine/Plugins/Runtime/AppleARKit/Source/AppleARKit/Private/AppleARKitPlaneAnchor.cpp
+++ b/Engine/Plugins/Runtime/AppleARKit/Source/AppleARKit/Private/AppleARKitPlaneAnchor.cpp
@@ -3,7 +3,7 @@
 // AppleARKit
 #include "AppleARKitPlaneAnchor.h"
 #include "AppleARKitModule.h"
-#include "AppleARKitTransform.h"
+#include "AppleARKitConversion.h"
 
 // UE4
 #include "Misc/ScopeLock.h"
diff --git a/Engine/Plugins/Runtime/AppleARKit/Source/AppleARKit/Private/AppleARKitSystem.cpp b/Engine/Plugins/Runtime/AppleARKit/Source/AppleARKit/Private/AppleARKitSystem.cpp
--- a/Engine/Plugins/Runtime/AppleARKit/Source/AppleARKit/Private/AppleARKitSystem.cpp
+++ b/Engine/Plugins/Runtime/AppleARKit/Source/AppleARKit/Private/AppleARKitSystem.cpp
@@ -5,7 +5,7 @@
 #include "AppleARKitSessionDelegate.h"
 #include "Misc/ScopeLock.h"
 #include "AppleARKitModule.h"
-#include "AppleARKitTransform.h"
+#include "AppleARKitConversion.h"
 #include "AppleARKitVideoOverlay.h"
 #include "AppleARKitFrame.h"
 #include "AppleARKitAnchor.h"
@@ -1204,7 +1204,6 @@ void FAppleARKitSystem::SessionDidFailWithError_DelegateThread(const FString& Er
 
 static TSharedPtr<FAppleARKitAnchorData> MakeAnchorData( ARAnchor* Anchor, const float WorldToMeterScale )
 {
-	// Construct appropriate UAppleARKitAnchor subclass
 	TSharedPtr<FAppleARKitAnchorData> NewAnchor;
 	if ([Anchor isKindOfClass:[ARPlaneAnchor class]])
 	{
diff --git a/Engine/Plugins/Runtime/AppleARKit/Source/AppleARKit/Private/AppleARKitTransform.h b/Engine/Plugins/Runtime/AppleARKit/Source/AppleARKit/Private/AppleARKitTransform.h
deleted file mode 100644
diff --git a/Engine/Plugins/Runtime/AppleARKit/Source/AppleARKit/Public/AppleARKitConversion.h b/Engine/Plugins/Runtime/AppleARKit/Source/AppleARKit/Public/AppleARKitConversion.h
new file mode 100644
--- /dev/null
+++ b/Engine/Plugins/Runtime/AppleARKit/Source/AppleARKit/Public/AppleARKitConversion.h
@@ -0,0 +1,48 @@
+// Copyright 1998-2018 Epic Games, Inc. All Rights Reserved.
+
+#pragma once
+
+// UE4
+#include "Math/Transform.h"
+
+struct FAppleARKitTransform
+{
+
+#if SUPPORTS_ARKIT_1_0
+
+	/** 
+	 * Convert's an ARKit 'Y up' 'right handed' coordinate system transform to Unreal's 'Z up' 
+	 * 'left handed' coordinate system.
+	 *
+	 * Ignores scale.
+	 */
+	static FORCEINLINE FTransform ToFTransform( const matrix_float4x4& RawYUpMatrix, float WorldToMetersScale = 100.0f )
+	{
+		// Conversion here is as per SteamVRHMD::ToFMatrix
+		FMatrix RawYUpFMatrix(
+			FPlane(RawYUpMatrix.columns[0][0], RawYUpMatrix.columns[0][1], RawYUpMatrix.columns[0][2], RawYUpMatrix.columns[0][3]),
+			FPlane(RawYUpMatrix.columns[1][0], RawYUpMatrix.columns[1][1], RawYUpMatrix.columns[1][2], RawYUpMatrix.columns[1][3]),
+			FPlane(RawYUpMatrix.columns[2][0], RawYUpMatrix.columns[2][1], RawYUpMatrix.columns[2][2], RawYUpMatrix.columns[2][3]),
+			FPlane(RawYUpMatrix.columns[3][0], RawYUpMatrix.columns[3][1], RawYUpMatrix.columns[3][2], RawYUpMatrix.columns[3][3]));
+
+		// Extract & convert translation
+		FVector Translation = FVector( -RawYUpFMatrix.M[3][2], RawYUpFMatrix.M[3][0], RawYUpFMatrix.M[3][1] ) * WorldToMetersScale;
+
+		// Extract & convert rotation 
+		FQuat RawRotation( RawYUpFMatrix );
+		FQuat Rotation( -RawRotation.Z, RawRotation.X, RawRotation.Y, -RawRotation.W );
+
+		return FTransform( Rotation, Translation );
+	}
+
+	/** 
+	 * Convert's an ARKit 'Y up' 'right handed' coordinate system vector to Unreal's 'Z up' 
+	 * 'left handed' coordinate system.
+	 */
+	static FORCEINLINE FVector ToFVector( const vector_float3& RawYUpVector, float WorldToMetersScale = 100.0f )
+	{
+		return FVector( -RawYUpVector.z, RawYUpVector.x, RawYUpVector.y ) * WorldToMetersScale;
+	}
+
+#endif
+};
