Prep work for moving face ar support to its own plugin

diff --git a/Engine/Plugins/Runtime/AppleARKit/Source/AppleARKit/Private/AppleARKitSystem.cpp b/Engine/Plugins/Runtime/AppleARKit/Source/AppleARKit/Private/AppleARKitSystem.cpp
--- a/Engine/Plugins/Runtime/AppleARKit/Source/AppleARKit/Private/AppleARKitSystem.cpp
+++ b/Engine/Plugins/Runtime/AppleARKit/Source/AppleARKit/Private/AppleARKitSystem.cpp
@@ -326,67 +326,6 @@ FARBlendShapeMap ToBlendShapeMap(NSDictionary<ARBlendShapeLocation,NSNumber *>*
 	return BlendShapeMap;
 }
 
-enum class EAppleAnchorType : uint8
-{
-	Anchor,
-	PlaneAnchor,
-	FaceAnchor,
-	ImageAnchor,
-	MAX
-};
-
-struct FAppleARKitAnchorData
-{
-	FAppleARKitAnchorData(FGuid InAnchorGuid, FTransform InTransform)
-	: Transform( InTransform )
-	, AnchorType( EAppleAnchorType::Anchor )
-	, AnchorGUID( InAnchorGuid )
-	{
-	}
-
-	FAppleARKitAnchorData(FGuid InAnchorGuid, FTransform InTransform, FVector InCenter, FVector InExtent)
-	: Transform( InTransform )
-	, AnchorType( EAppleAnchorType::PlaneAnchor )
-	, AnchorGUID( InAnchorGuid )
-	, Center(InCenter)
-	, Extent(InExtent)
-	{
-	}
-
-	FAppleARKitAnchorData(FGuid InAnchorGuid, FTransform InTransform, FARBlendShapeMap InBlendShapes, TArray<FVector> InFaceVerts)
-	: Transform( InTransform )
-	, AnchorType( EAppleAnchorType::FaceAnchor )
-	, AnchorGUID( InAnchorGuid )
-	, BlendShapes( MoveTemp(InBlendShapes) )
-	, FaceVerts( MoveTemp(InFaceVerts) )
-	{
-	}
-
-    FAppleARKitAnchorData(FGuid InAnchorGuid, FTransform InTransform, FString InDetectedImageName)
-    : Transform( InTransform )
-    , AnchorType( EAppleAnchorType::ImageAnchor )
-    , AnchorGUID( InAnchorGuid )
-    , DetectedImageName( MoveTemp(InDetectedImageName) )
-    {
-    }
-
-	FTransform Transform;
-	EAppleAnchorType AnchorType;
-	FGuid AnchorGUID;
-	FVector Center;
-	FVector Extent;
-
-	TArray<FVector> BoundaryVerts;
-	FARBlendShapeMap BlendShapes;
-	TArray<FVector> FaceVerts;
-	// Note: the index buffer never changes so can be safely read once
-	static TArray<int32> FaceIndices;
-
-	FString DetectedImageName;
-};
-
-TArray<int32> FAppleARKitAnchorData::FaceIndices;
-
 #endif
 
 
diff --git a/Engine/Plugins/Runtime/AppleARKit/Source/AppleARKit/Public/AppleARKitConversion.h b/Engine/Plugins/Runtime/AppleARKit/Source/AppleARKit/Public/AppleARKitConversion.h
--- a/Engine/Plugins/Runtime/AppleARKit/Source/AppleARKit/Public/AppleARKitConversion.h
+++ b/Engine/Plugins/Runtime/AppleARKit/Source/AppleARKit/Public/AppleARKitConversion.h
@@ -46,3 +46,65 @@ struct FAppleARKitTransform
 
 #endif
 };
+
+enum class EAppleAnchorType : uint8
+{
+	Anchor,
+	PlaneAnchor,
+	FaceAnchor,
+	ImageAnchor,
+	MAX
+};
+
+struct FAppleARKitAnchorData
+{
+	FAppleARKitAnchorData(FGuid InAnchorGuid, FTransform InTransform)
+		: Transform( InTransform )
+		, AnchorType( EAppleAnchorType::Anchor )
+		, AnchorGUID( InAnchorGuid )
+	{
+	}
+
+	FAppleARKitAnchorData(FGuid InAnchorGuid, FTransform InTransform, FVector InCenter, FVector InExtent)
+		: Transform( InTransform )
+		, AnchorType( EAppleAnchorType::PlaneAnchor )
+		, AnchorGUID( InAnchorGuid )
+		, Center(InCenter)
+		, Extent(InExtent)
+	{
+	}
+
+	FAppleARKitAnchorData(FGuid InAnchorGuid, FTransform InTransform, FARBlendShapeMap InBlendShapes, TArray<FVector> InFaceVerts)
+		: Transform( InTransform )
+		, AnchorType( EAppleAnchorType::FaceAnchor )
+		, AnchorGUID( InAnchorGuid )
+		, BlendShapes( MoveTemp(InBlendShapes) )
+		, FaceVerts( MoveTemp(InFaceVerts) )
+	{
+	}
+
+	FAppleARKitAnchorData(FGuid InAnchorGuid, FTransform InTransform, FString InDetectedImageName)
+		: Transform( InTransform )
+		, AnchorType( EAppleAnchorType::ImageAnchor )
+		, AnchorGUID( InAnchorGuid )
+		, DetectedImageName( MoveTemp(InDetectedImageName) )
+	{
+	}
+
+	FTransform Transform;
+	EAppleAnchorType AnchorType;
+	FGuid AnchorGUID;
+	FVector Center;
+	FVector Extent;
+
+	TArray<FVector> BoundaryVerts;
+	FARBlendShapeMap BlendShapes;
+	TArray<FVector> FaceVerts;
+	// Temp non-static while code is rearranged
+	TArray<int32> FaceIndices;
+	// Note: the index buffer never changes so can be safely read once
+//	static TArray<int32> FaceIndices;
+//	TArray<int32> FAppleARKitAnchorData::FaceIndices;
+
+	FString DetectedImageName;
+};
