Pulled face ar support into its own plugin so that it can be enabled/disabled independently from room ar

diff --git a/Engine/Config/BaseEngine.ini b/Engine/Config/BaseEngine.ini
--- a/Engine/Config/BaseEngine.ini
+++ b/Engine/Config/BaseEngine.ini
@@ -564,6 +564,7 @@ bUseStreamingPause=false
 +ClassRedirects=(OldName="/Script/AssetScriptingUtilities.SkeletalMeshUtilitiesLibrary", NewName="/Script/EditorScriptingUtilities.EditorSkeletalMeshLibrary")
 +FunctionRedirects=(OldName="StaticMeshUtilitiesLibrary.GetLODScreenSizes",NewName="EditorStaticMeshLibrary.GetLodScreenSizes")
 
++ClassRedirects=(OldName="AppleARKitFaceMeshComponent",NewName="/Script/AppleARKitFaceSupport.AppleARKitFaceMeshComponent")
 
 [CoreUObject.Metadata]
 ; Note: UnrealHeaderTool should be rerun after making changes to MetadataRedirects to catch any keys specified in class headers
diff --git a/Engine/Plugins/Runtime/AppleARKit/AppleARKit.uplugin b/Engine/Plugins/Runtime/AppleARKit/AppleARKit.uplugin
--- a/Engine/Plugins/Runtime/AppleARKit/AppleARKit.uplugin
+++ b/Engine/Plugins/Runtime/AppleARKit/AppleARKit.uplugin
@@ -12,7 +12,7 @@
 	"SupportURL": "",
 	"CanContainContent": true,
 	"IsBetaVersion": true,
-	"EnabledByDefault": true,
+	"EnabledByDefault": false,
 	"Modules": [
 		{
 			"Name": "AppleARKit",
@@ -22,14 +22,6 @@
 		}
 	],
     "Plugins": [
-		{
-			"Name": "ProceduralMeshComponent",
-			"Enabled": true
-		},
-        {
-            "Name": "LiveLink",
-            "Enabled": true
-        },
         {
             "Name": "AppleImageUtils",
             "Enabled": true
diff --git a/Engine/Plugins/Runtime/AppleARKit/Source/AppleARKit/AppleARKit.Build.cs b/Engine/Plugins/Runtime/AppleARKit/Source/AppleARKit/AppleARKit.Build.cs
--- a/Engine/Plugins/Runtime/AppleARKit/Source/AppleARKit/AppleARKit.Build.cs
+++ b/Engine/Plugins/Runtime/AppleARKit/Source/AppleARKit/AppleARKit.Build.cs
@@ -37,11 +37,6 @@ public class AppleARKit : ModuleRules
                 "HeadMountedDisplay",
                 "IOSRuntimeSettings",
                 "AugmentedReality",
-                "ProceduralMeshComponent",
-                "LiveLink",
-                "LiveLinkInterface",
-//                "OnlineSubsystem",
-                "Sockets",
                 "AppleImageUtils"
 				// ... add private dependencies that you statically link with here ...
 			}
diff --git a/Engine/Plugins/Runtime/AppleARKit/Source/AppleARKit/Private/AppleARKitAnchor.cpp b/Engine/Plugins/Runtime/AppleARKit/Source/AppleARKit/Private/AppleARKitAnchor.cpp
--- a/Engine/Plugins/Runtime/AppleARKit/Source/AppleARKit/Private/AppleARKitAnchor.cpp
+++ b/Engine/Plugins/Runtime/AppleARKit/Source/AppleARKit/Private/AppleARKitAnchor.cpp
@@ -22,7 +22,7 @@ void UAppleARKitAnchor::Update_DelegateThread( ARAnchor* Anchor )
 	FScopeLock ScopeLock( &UpdateLock );
 
 	// @todo arkit use World Settings WorldToMetersScale
-	Transform = FAppleARKitTransform::ToFTransform( Anchor.transform );
+	Transform = FAppleARKitConversion::ToFTransform( Anchor.transform );
 }
 
 #endif
diff --git a/Engine/Plugins/Runtime/AppleARKit/Source/AppleARKit/Private/AppleARKitCamera.cpp b/Engine/Plugins/Runtime/AppleARKit/Source/AppleARKit/Private/AppleARKitCamera.cpp
--- a/Engine/Plugins/Runtime/AppleARKit/Source/AppleARKit/Private/AppleARKitCamera.cpp
+++ b/Engine/Plugins/Runtime/AppleARKit/Source/AppleARKit/Private/AppleARKitCamera.cpp
@@ -45,7 +45,7 @@ FAppleARKitCamera::FAppleARKitCamera( ARCamera* InARCamera )
 	
 	
     // Copy / convert camera transform
-    Transform = FAppleARKitTransform::ToFTransform( InARCamera.transform );
+    Transform = FAppleARKitConversion::ToFTransform( InARCamera.transform );
     Translation = Transform.GetTranslation();
 	Orientation = Transform.GetRotation();
 	
diff --git a/Engine/Plugins/Runtime/AppleARKit/Source/AppleARKit/Private/AppleARKitConfiguration.cpp b/Engine/Plugins/Runtime/AppleARKit/Source/AppleARKit/Private/AppleARKitConfiguration.cpp
--- a/Engine/Plugins/Runtime/AppleARKit/Source/AppleARKit/Private/AppleARKitConfiguration.cpp
+++ b/Engine/Plugins/Runtime/AppleARKit/Source/AppleARKit/Private/AppleARKitConfiguration.cpp
@@ -11,21 +11,6 @@
 #pragma clang diagnostic push
 #pragma clang diagnostic ignored "-Wpartial-availability"
 
-ARWorldAlignment ToARWorldAlignment( const EAppleARKitWorldAlignment& InWorldAlignment )
-{
-	switch ( InWorldAlignment )
-	{
-		case EAppleARKitWorldAlignment::Gravity:
-			return ARWorldAlignmentGravity;
-			
-		case EAppleARKitWorldAlignment::GravityAndHeading:
-			return ARWorldAlignmentGravityAndHeading;
-
-		case EAppleARKitWorldAlignment::Camera:
-			return ARWorldAlignmentCamera;
-	};
-};
-
 #if SUPPORTS_ARKIT_1_5
 
 void InitImageDetection(UARSessionConfig* SessionConfig, ARWorldTrackingConfiguration* WorldConfig, TMap< FString, UARCandidateImage* >& CandidateImages, TMap< FString, CGImageRef >& ConvertedCandidateImages)
@@ -124,15 +109,6 @@ ARConfiguration* ToARConfiguration( UARSessionConfig* SessionConfig, FAppleARKit
 			SessionConfiguration = WorldTrackingConfiguration;
 			break;
 		}
-		case EARSessionType::Face:
-		{
-			if (ARFaceTrackingConfiguration.isSupported == FALSE)
-			{
-				return nullptr;
-			}
-			SessionConfiguration = [ARFaceTrackingConfiguration new];
-			break;
-		}
 		default:
 			return nullptr;
 	}
@@ -141,7 +117,7 @@ ARConfiguration* ToARConfiguration( UARSessionConfig* SessionConfig, FAppleARKit
 	// Copy / convert properties
 	SessionConfiguration.lightEstimationEnabled = InConfiguration.bLightEstimationEnabled;
 	SessionConfiguration.providesAudioData = InConfiguration.bProvidesAudioData;
-	SessionConfiguration.worldAlignment = ToARWorldAlignment(InConfiguration.Alignment);
+	SessionConfiguration.worldAlignment = FAppleARKitConversion::ToARWorldAlignment(InConfiguration.Alignment);
 
     return SessionConfiguration;
 }
diff --git a/Engine/Plugins/Runtime/AppleARKit/Source/AppleARKit/Private/AppleARKitConfiguration.h b/Engine/Plugins/Runtime/AppleARKit/Source/AppleARKit/Private/AppleARKitConfiguration.h
--- a/Engine/Plugins/Runtime/AppleARKit/Source/AppleARKit/Private/AppleARKitConfiguration.h
+++ b/Engine/Plugins/Runtime/AppleARKit/Source/AppleARKit/Private/AppleARKitConfiguration.h
@@ -11,41 +11,9 @@
 #endif
 
 #include "ARSystem.h"
-
-/**
- * Enum constants for indicating the world alignment.
- */
-enum class EAppleARKitWorldAlignment : uint8
-{
-	/** Aligns the world with gravity that is defined by vector (0, -1, 0). */
-	Gravity,
-	
-	/**
-	 * Aligns the world with gravity that is defined by the vector (0, -1, 0)
-	 * and heading (w.r.t. True North) that is given by the vector (0, 0, -1).
-	 */
-	GravityAndHeading,
-	
-	/** Aligns the world with the camera's orientation. */
-	Camera
-};
-ENUM_CLASS_FLAGS(EAppleARKitWorldAlignment)
-
+#include "AppleARKitConversion.h"
 
 
-/**
- Option set indicating the type of planes to detect.
- */
-enum class EAppleARKitPlaneDetection : uint8
-{
-	/** No plane detection is run. */
-	None        = 1,
-	
-	/** Plane detection determines horizontal planes in the scene. */
-	Horizontal  = 2
-};
-ENUM_CLASS_FLAGS(EAppleARKitPlaneDetection)
-
 #if SUPPORTS_ARKIT_1_0
 ARConfiguration* ToARConfiguration( UARSessionConfig* SessionConfig, class FAppleARKitConfiguration& InConfiguration, TMap< FString, UARCandidateImage* >& CandidateImages, TMap< FString, CGImageRef >& ConvertedCandidateImages );
 #endif
@@ -75,21 +43,5 @@ public:
 	 * @discussion The default is ARWorldAlignmentGravity.
 	 */
 	EAppleARKitWorldAlignment Alignment = EAppleARKitWorldAlignment::Gravity;
-	
-	/**
-	 * A session configuration for world tracking.
-	 *
-	 * World tracking provides 6 degrees of freedom tracking of the device.
-	 * By finding feature points in the scene, world tracking enables performing hit-tests against the frame.
-	 * Tracking can no longer be resumed once the session is paused.
-	 */
-	struct FWorldTracking
-	{
-		/**
-		 * Plane detection settings.
-		 */
-		EAppleARKitPlaneDetection PlaneDetection = EAppleARKitPlaneDetection::Horizontal;
-	};
-	TOptional<FWorldTracking> WorldTracking = FWorldTracking();
 };
 
diff --git a/Engine/Plugins/Runtime/AppleARKit/Source/AppleARKit/Private/AppleARKitFaceMeshComponent.cpp b/Engine/Plugins/Runtime/AppleARKit/Source/AppleARKit/Private/AppleARKitFaceMeshComponent.cpp
deleted file mode 100644
diff --git a/Engine/Plugins/Runtime/AppleARKit/Source/AppleARKit/Private/AppleARKitFaceMeshConversion.h b/Engine/Plugins/Runtime/AppleARKit/Source/AppleARKit/Private/AppleARKitFaceMeshConversion.h
deleted file mode 100644
diff --git a/Engine/Plugins/Runtime/AppleARKit/Source/AppleARKit/Private/AppleARKitHitTestResult.cpp b/Engine/Plugins/Runtime/AppleARKit/Source/AppleARKit/Private/AppleARKitHitTestResult.cpp
--- a/Engine/Plugins/Runtime/AppleARKit/Source/AppleARKit/Private/AppleARKitHitTestResult.cpp
+++ b/Engine/Plugins/Runtime/AppleARKit/Source/AppleARKit/Private/AppleARKitHitTestResult.cpp
@@ -70,7 +70,7 @@ FAppleARKitHitTestResult::FAppleARKitHitTestResult( ARHitTestResult* InARHitTest
 	// Convert properties
 	Type = ToEAppleARKitHitTestResultType( InARHitTestResult.type );
     Distance = InARHitTestResult.distance * WorldToMetersScale;
-	Transform = FAppleARKitTransform::ToFTransform( InARHitTestResult.worldTransform, WorldToMetersScale );
+	Transform = FAppleARKitConversion::ToFTransform( InARHitTestResult.worldTransform );
 	Anchor = InAnchor;
 }
 
diff --git a/Engine/Plugins/Runtime/AppleARKit/Source/AppleARKit/Private/AppleARKitLiveLinkSource.cpp b/Engine/Plugins/Runtime/AppleARKit/Source/AppleARKit/Private/AppleARKitLiveLinkSource.cpp
deleted file mode 100644
diff --git a/Engine/Plugins/Runtime/AppleARKit/Source/AppleARKit/Private/AppleARKitLiveLinkSource.h b/Engine/Plugins/Runtime/AppleARKit/Source/AppleARKit/Private/AppleARKitLiveLinkSource.h
deleted file mode 100644
diff --git a/Engine/Plugins/Runtime/AppleARKit/Source/AppleARKit/Private/AppleARKitModule.cpp b/Engine/Plugins/Runtime/AppleARKit/Source/AppleARKit/Private/AppleARKitModule.cpp
--- a/Engine/Plugins/Runtime/AppleARKit/Source/AppleARKit/Private/AppleARKitModule.cpp
+++ b/Engine/Plugins/Runtime/AppleARKit/Source/AppleARKit/Private/AppleARKitModule.cpp
@@ -39,16 +39,6 @@ void FAppleARKitModule::StartupModule()
 	IHeadMountedDisplayModule::StartupModule();
 
 	FCoreDelegates::OnPreExit.AddRaw(this, &FAppleARKitModule::PreExit);
-
-	// LiveLink listener needs to be created here so that the editor can receive remote publishing events
-#if PLATFORM_DESKTOP
-	bool bEnableLiveLinkForFaceTracking = false;
-	GConfig->GetBool(TEXT("/Script/AppleARKit.AppleARKitSettings"), TEXT("bEnableLiveLinkForFaceTracking"), bEnableLiveLinkForFaceTracking, GEngineIni);
-	if (bEnableLiveLinkForFaceTracking)
-	{
-		FAppleARKitLiveLinkSourceFactory::CreateLiveLinkRemoteListener();
-	}
-#endif
 }
 
 void FAppleARKitModule::PreExit()
diff --git a/Engine/Plugins/Runtime/AppleARKit/Source/AppleARKit/Private/AppleARKitPlaneAnchor.cpp b/Engine/Plugins/Runtime/AppleARKit/Source/AppleARKit/Private/AppleARKitPlaneAnchor.cpp
--- a/Engine/Plugins/Runtime/AppleARKit/Source/AppleARKit/Private/AppleARKitPlaneAnchor.cpp
+++ b/Engine/Plugins/Runtime/AppleARKit/Source/AppleARKit/Private/AppleARKitPlaneAnchor.cpp
@@ -44,8 +44,8 @@ void UAppleARKitPlaneAnchor::Update_DelegateThread( ARAnchor* Anchor )
 		FScopeLock ScopeLock( &UpdateLock );
 		
 		// @todo use World Settings WorldToMetersScale
-		Extent = FAppleARKitTransform::ToFVector( PlaneAnchor.extent ).GetAbs();
-		Center = FAppleARKitTransform::ToFVector( PlaneAnchor.center );
+		Extent = FAppleARKitConversion::ToFVector( PlaneAnchor.extent ).GetAbs();
+		Center = FAppleARKitConversion::ToFVector( PlaneAnchor.center );
 	}
 }
 
diff --git a/Engine/Plugins/Runtime/AppleARKit/Source/AppleARKit/Private/AppleARKitSystem.cpp b/Engine/Plugins/Runtime/AppleARKit/Source/AppleARKit/Private/AppleARKitSystem.cpp
--- a/Engine/Plugins/Runtime/AppleARKit/Source/AppleARKit/Private/AppleARKitSystem.cpp
+++ b/Engine/Plugins/Runtime/AppleARKit/Source/AppleARKit/Private/AppleARKitSystem.cpp
@@ -13,7 +13,6 @@
 #include "AppleARKitConfiguration.h"
 #include "GeneralProjectSettings.h"
 #include "IOSRuntimeSettings.h"
-#include "AppleARKitFaceMeshConversion.h"
 #include "ARSessionConfig.h"
 #include "AppleARKitSettings.h"
 #include "ARTrackable.h"
@@ -21,6 +20,9 @@
 #include "ARTraceResult.h"
 #include "ARPin.h"
 
+// To separate out the face ar library linkage from standard ar apps
+#include "AppleARKitFaceSupport.h"
+
 // For orientation changed
 #include "Misc/CoreDelegates.h"
 
@@ -29,40 +31,6 @@
 	#pragma clang diagnostic ignored "-Wpartial-availability"
 #endif
 
-namespace ARKitUtil
-{
-
-static UARPin* PinFromComponent( const USceneComponent* Component, const TArray<UARPin*>& InPins )
-{
-	for (UARPin* Pin : InPins)
-	{
-		if (Pin->GetPinnedComponent() == Component)
-		{
-			return Pin;
-		}
-	}
-	
-	return nullptr;
-}
-
-
-static TArray<UARPin*> PinsFromGeometry( const UARTrackedGeometry* Geometry, const TArray<UARPin*>& InPins )
-{
-	TArray<UARPin*> OutPins;
-	for (UARPin* Pin : InPins)
-	{
-		if (Pin->GetTrackedGeometry() == Geometry)
-		{
-			OutPins.Add(Pin);
-		}
-	}
-	
-	return OutPins;
-}
-
-	
-}
-
 
 //
 //  FAppleARKitXRCamera
@@ -174,165 +142,6 @@ private:
 };
 
 
-
-#if SUPPORTS_ARKIT_1_0
-static FORCEINLINE FGuid ToFGuid( uuid_t UUID )
-{
-	FGuid AsGUID(
-		*(uint32*)UUID,
-		*((uint32*)UUID)+1,
-		*((uint32*)UUID)+2,
-		*((uint32*)UUID)+3);
-	return AsGUID;
-}
-
-static FORCEINLINE FGuid ToFGuid( NSUUID* Identifier )
-{
-	// Get bytes
-	uuid_t UUID;
-	[Identifier getUUIDBytes:UUID];
-	
-	// Set FGuid parts
-	return ToFGuid( UUID );
-}
-
-FARBlendShapeMap ToBlendShapeMap(NSDictionary<ARBlendShapeLocation,NSNumber *>* BlendShapes, const FTransform& Transform)
-{
-	FARBlendShapeMap BlendShapeMap;
-	FRotator TrackedRot(Transform.GetRotation());
-	// Map the -180..180 range to -1..1
-	float HeadYaw = (float)TrackedRot.Yaw / 180.f;
-	float HeadPitch = (float)TrackedRot.Pitch / 180.f;
-	float HeadRoll = (float)TrackedRot.Roll / 180.f;
-	BlendShapeMap.Add(EARFaceBlendShape::HeadYaw, HeadYaw);
-	BlendShapeMap.Add(EARFaceBlendShape::HeadPitch, HeadPitch);
-	BlendShapeMap.Add(EARFaceBlendShape::HeadRoll, HeadRoll);
-
-#define SET_BLEND_SHAPE(AppleShape, UE4Shape) \
-	if (BlendShapes[AppleShape]) \
-	{ \
-		BlendShapeMap.Add(UE4Shape, [BlendShapes[AppleShape] floatValue]); \
-	} \
-	else \
-	{ \
-		BlendShapeMap.Add(UE4Shape, 0.f); \
-	}
-
-	// Do we want to capture face performance or look at the face as if in a mirror
-	if (GetDefault<UAppleARKitSettings>()->DefaultFaceTrackingDirection == EARFaceTrackingDirection::FaceRelative)
-	{
-		SET_BLEND_SHAPE(ARBlendShapeLocationEyeBlinkLeft, EARFaceBlendShape::EyeBlinkLeft);
-		SET_BLEND_SHAPE(ARBlendShapeLocationEyeLookDownLeft, EARFaceBlendShape::EyeLookDownLeft);
-		SET_BLEND_SHAPE(ARBlendShapeLocationEyeLookInLeft, EARFaceBlendShape::EyeLookInLeft);
-		SET_BLEND_SHAPE(ARBlendShapeLocationEyeLookOutLeft, EARFaceBlendShape::EyeLookOutLeft);
-		SET_BLEND_SHAPE(ARBlendShapeLocationEyeLookUpLeft, EARFaceBlendShape::EyeLookUpLeft);
-		SET_BLEND_SHAPE(ARBlendShapeLocationEyeSquintLeft, EARFaceBlendShape::EyeSquintLeft);
-		SET_BLEND_SHAPE(ARBlendShapeLocationEyeWideLeft, EARFaceBlendShape::EyeWideLeft);
-		SET_BLEND_SHAPE(ARBlendShapeLocationEyeBlinkRight, EARFaceBlendShape::EyeBlinkRight);
-		SET_BLEND_SHAPE(ARBlendShapeLocationEyeLookDownRight, EARFaceBlendShape::EyeLookDownRight);
-		SET_BLEND_SHAPE(ARBlendShapeLocationEyeLookInRight, EARFaceBlendShape::EyeLookInRight);
-		SET_BLEND_SHAPE(ARBlendShapeLocationEyeLookOutRight, EARFaceBlendShape::EyeLookOutRight);
-		SET_BLEND_SHAPE(ARBlendShapeLocationEyeLookUpRight, EARFaceBlendShape::EyeLookUpRight);
-		SET_BLEND_SHAPE(ARBlendShapeLocationEyeSquintRight, EARFaceBlendShape::EyeSquintRight);
-		SET_BLEND_SHAPE(ARBlendShapeLocationEyeWideRight, EARFaceBlendShape::EyeWideRight);
-		SET_BLEND_SHAPE(ARBlendShapeLocationJawForward, EARFaceBlendShape::JawForward);
-		SET_BLEND_SHAPE(ARBlendShapeLocationJawLeft, EARFaceBlendShape::JawLeft);
-		SET_BLEND_SHAPE(ARBlendShapeLocationJawRight, EARFaceBlendShape::JawRight);
-		SET_BLEND_SHAPE(ARBlendShapeLocationMouthLeft, EARFaceBlendShape::MouthLeft);
-		SET_BLEND_SHAPE(ARBlendShapeLocationMouthRight, EARFaceBlendShape::MouthRight);
-		SET_BLEND_SHAPE(ARBlendShapeLocationMouthSmileLeft, EARFaceBlendShape::MouthSmileLeft);
-		SET_BLEND_SHAPE(ARBlendShapeLocationMouthSmileRight, EARFaceBlendShape::MouthSmileRight);
-		SET_BLEND_SHAPE(ARBlendShapeLocationMouthFrownLeft, EARFaceBlendShape::MouthFrownLeft);
-		SET_BLEND_SHAPE(ARBlendShapeLocationMouthFrownRight, EARFaceBlendShape::MouthFrownRight);
-		SET_BLEND_SHAPE(ARBlendShapeLocationMouthDimpleLeft, EARFaceBlendShape::MouthDimpleLeft);
-		SET_BLEND_SHAPE(ARBlendShapeLocationMouthDimpleRight, EARFaceBlendShape::MouthDimpleRight);
-		SET_BLEND_SHAPE(ARBlendShapeLocationMouthStretchLeft, EARFaceBlendShape::MouthStretchLeft);
-		SET_BLEND_SHAPE(ARBlendShapeLocationMouthStretchRight, EARFaceBlendShape::MouthStretchRight);
-		SET_BLEND_SHAPE(ARBlendShapeLocationMouthPressLeft, EARFaceBlendShape::MouthPressLeft);
-		SET_BLEND_SHAPE(ARBlendShapeLocationMouthPressRight, EARFaceBlendShape::MouthPressRight);
-		SET_BLEND_SHAPE(ARBlendShapeLocationMouthLowerDownLeft, EARFaceBlendShape::MouthLowerDownLeft);
-		SET_BLEND_SHAPE(ARBlendShapeLocationMouthLowerDownRight, EARFaceBlendShape::MouthLowerDownRight);
-		SET_BLEND_SHAPE(ARBlendShapeLocationMouthUpperUpLeft, EARFaceBlendShape::MouthUpperUpLeft);
-		SET_BLEND_SHAPE(ARBlendShapeLocationMouthUpperUpRight, EARFaceBlendShape::MouthUpperUpRight);
-		SET_BLEND_SHAPE(ARBlendShapeLocationBrowDownLeft, EARFaceBlendShape::BrowDownLeft);
-		SET_BLEND_SHAPE(ARBlendShapeLocationBrowDownRight, EARFaceBlendShape::BrowDownRight);
-		SET_BLEND_SHAPE(ARBlendShapeLocationBrowOuterUpLeft, EARFaceBlendShape::BrowOuterUpLeft);
-		SET_BLEND_SHAPE(ARBlendShapeLocationBrowOuterUpRight, EARFaceBlendShape::BrowOuterUpRight);
-		SET_BLEND_SHAPE(ARBlendShapeLocationCheekSquintLeft, EARFaceBlendShape::CheekSquintLeft);
-		SET_BLEND_SHAPE(ARBlendShapeLocationCheekSquintRight, EARFaceBlendShape::CheekSquintRight);
-		SET_BLEND_SHAPE(ARBlendShapeLocationNoseSneerLeft, EARFaceBlendShape::NoseSneerLeft);
-		SET_BLEND_SHAPE(ARBlendShapeLocationNoseSneerRight, EARFaceBlendShape::NoseSneerRight);
-	}
-	else
-	{
-		SET_BLEND_SHAPE(ARBlendShapeLocationEyeBlinkLeft, EARFaceBlendShape::EyeBlinkRight);
-		SET_BLEND_SHAPE(ARBlendShapeLocationEyeLookDownLeft, EARFaceBlendShape::EyeLookDownRight);
-		SET_BLEND_SHAPE(ARBlendShapeLocationEyeLookInLeft, EARFaceBlendShape::EyeLookInRight);
-		SET_BLEND_SHAPE(ARBlendShapeLocationEyeLookOutLeft, EARFaceBlendShape::EyeLookOutRight);
-		SET_BLEND_SHAPE(ARBlendShapeLocationEyeLookUpLeft, EARFaceBlendShape::EyeLookUpRight);
-		SET_BLEND_SHAPE(ARBlendShapeLocationEyeSquintLeft, EARFaceBlendShape::EyeSquintRight);
-		SET_BLEND_SHAPE(ARBlendShapeLocationEyeWideLeft, EARFaceBlendShape::EyeWideRight);
-		SET_BLEND_SHAPE(ARBlendShapeLocationEyeBlinkRight, EARFaceBlendShape::EyeBlinkLeft);
-		SET_BLEND_SHAPE(ARBlendShapeLocationEyeLookDownRight, EARFaceBlendShape::EyeLookDownLeft);
-		SET_BLEND_SHAPE(ARBlendShapeLocationEyeLookInRight, EARFaceBlendShape::EyeLookInLeft);
-		SET_BLEND_SHAPE(ARBlendShapeLocationEyeLookOutRight, EARFaceBlendShape::EyeLookOutLeft);
-		SET_BLEND_SHAPE(ARBlendShapeLocationEyeLookUpRight, EARFaceBlendShape::EyeLookUpLeft);
-		SET_BLEND_SHAPE(ARBlendShapeLocationEyeSquintRight, EARFaceBlendShape::EyeSquintLeft);
-		SET_BLEND_SHAPE(ARBlendShapeLocationEyeWideRight, EARFaceBlendShape::EyeWideLeft);
-		SET_BLEND_SHAPE(ARBlendShapeLocationJawForward, EARFaceBlendShape::JawForward);
-		SET_BLEND_SHAPE(ARBlendShapeLocationJawLeft, EARFaceBlendShape::JawRight);
-		SET_BLEND_SHAPE(ARBlendShapeLocationJawRight, EARFaceBlendShape::JawLeft);
-		SET_BLEND_SHAPE(ARBlendShapeLocationMouthLeft, EARFaceBlendShape::MouthRight);
-		SET_BLEND_SHAPE(ARBlendShapeLocationMouthRight, EARFaceBlendShape::MouthLeft);
-		SET_BLEND_SHAPE(ARBlendShapeLocationMouthSmileLeft, EARFaceBlendShape::MouthSmileRight);
-		SET_BLEND_SHAPE(ARBlendShapeLocationMouthSmileRight, EARFaceBlendShape::MouthSmileLeft);
-		SET_BLEND_SHAPE(ARBlendShapeLocationMouthFrownLeft, EARFaceBlendShape::MouthFrownRight);
-		SET_BLEND_SHAPE(ARBlendShapeLocationMouthFrownRight, EARFaceBlendShape::MouthFrownLeft);
-		SET_BLEND_SHAPE(ARBlendShapeLocationMouthDimpleLeft, EARFaceBlendShape::MouthDimpleRight);
-		SET_BLEND_SHAPE(ARBlendShapeLocationMouthDimpleRight, EARFaceBlendShape::MouthDimpleLeft);
-		SET_BLEND_SHAPE(ARBlendShapeLocationMouthStretchLeft, EARFaceBlendShape::MouthStretchRight);
-		SET_BLEND_SHAPE(ARBlendShapeLocationMouthStretchRight, EARFaceBlendShape::MouthStretchLeft);
-		SET_BLEND_SHAPE(ARBlendShapeLocationMouthPressLeft, EARFaceBlendShape::MouthPressRight);
-		SET_BLEND_SHAPE(ARBlendShapeLocationMouthPressRight, EARFaceBlendShape::MouthPressLeft);
-		SET_BLEND_SHAPE(ARBlendShapeLocationMouthLowerDownLeft, EARFaceBlendShape::MouthLowerDownRight);
-		SET_BLEND_SHAPE(ARBlendShapeLocationMouthLowerDownRight, EARFaceBlendShape::MouthLowerDownLeft);
-		SET_BLEND_SHAPE(ARBlendShapeLocationMouthUpperUpLeft, EARFaceBlendShape::MouthUpperUpRight);
-		SET_BLEND_SHAPE(ARBlendShapeLocationMouthUpperUpRight, EARFaceBlendShape::MouthUpperUpLeft);
-		SET_BLEND_SHAPE(ARBlendShapeLocationBrowDownLeft, EARFaceBlendShape::BrowDownRight);
-		SET_BLEND_SHAPE(ARBlendShapeLocationBrowDownRight, EARFaceBlendShape::BrowDownLeft);
-		SET_BLEND_SHAPE(ARBlendShapeLocationBrowOuterUpLeft, EARFaceBlendShape::BrowOuterUpRight);
-		SET_BLEND_SHAPE(ARBlendShapeLocationBrowOuterUpRight, EARFaceBlendShape::BrowOuterUpLeft);
-		SET_BLEND_SHAPE(ARBlendShapeLocationCheekSquintLeft, EARFaceBlendShape::CheekSquintRight);
-		SET_BLEND_SHAPE(ARBlendShapeLocationCheekSquintRight, EARFaceBlendShape::CheekSquintLeft);
-		SET_BLEND_SHAPE(ARBlendShapeLocationNoseSneerLeft, EARFaceBlendShape::NoseSneerRight);
-		SET_BLEND_SHAPE(ARBlendShapeLocationNoseSneerRight, EARFaceBlendShape::NoseSneerLeft);
-	}
-
-	// These are the same mirrored or not
-	SET_BLEND_SHAPE(ARBlendShapeLocationJawOpen, EARFaceBlendShape::JawOpen);
-	SET_BLEND_SHAPE(ARBlendShapeLocationMouthClose, EARFaceBlendShape::MouthClose);
-	SET_BLEND_SHAPE(ARBlendShapeLocationMouthFunnel, EARFaceBlendShape::MouthFunnel);
-	SET_BLEND_SHAPE(ARBlendShapeLocationMouthPucker, EARFaceBlendShape::MouthPucker);
-	SET_BLEND_SHAPE(ARBlendShapeLocationMouthRollLower, EARFaceBlendShape::MouthRollLower);
-	SET_BLEND_SHAPE(ARBlendShapeLocationMouthRollUpper, EARFaceBlendShape::MouthRollUpper);
-	SET_BLEND_SHAPE(ARBlendShapeLocationMouthShrugLower, EARFaceBlendShape::MouthShrugLower);
-	SET_BLEND_SHAPE(ARBlendShapeLocationMouthShrugUpper, EARFaceBlendShape::MouthShrugUpper);
-	SET_BLEND_SHAPE(ARBlendShapeLocationBrowInnerUp, EARFaceBlendShape::BrowInnerUp);
-	SET_BLEND_SHAPE(ARBlendShapeLocationCheekPuff, EARFaceBlendShape::CheekPuff);
-
-
-#undef SET_BLEND_SHAPE
-
-	return BlendShapeMap;
-}
-
-#endif
-
-
-
-
-
-
 //
 //  FAppleARKitSystem
 //
@@ -349,18 +158,6 @@ FAppleARKitSystem::FAppleARKitSystem()
 , LastTrackedGeometry_DebugId(0)
 {
 	// See Initialize(), as we need access to SharedThis()
-
-	// Create our LiveLink provider if the project setting is enabled
-	if (GetDefault<UAppleARKitSettings>()->bEnableLiveLinkForFaceTracking)
-	{
-		FaceTrackingLiveLinkSubjectName = GetDefault<UAppleARKitSettings>()->DefaultFaceTrackingLiveLinkSubjectName;
-#if PLATFORM_IOS
-		LiveLinkSource = FAppleARKitLiveLinkSourceFactory::CreateLiveLinkSource(true);
-#else
-		// This should be started already, but just in case
-		FAppleARKitLiveLinkSourceFactory::CreateLiveLinkRemoteListener();
-#endif
-	}
 }
 
 FAppleARKitSystem::~FAppleARKitSystem()
@@ -373,6 +170,7 @@ void FAppleARKitSystem::Shutdown()
 #if SUPPORTS_ARKIT_1_0
 	if (Session != nullptr)
 	{
+		FaceARSupport = TSharedPtr<FAppleARKitFaceSupportBase, ESPMode::ThreadSafe>();
 		[Session pause];
 		Session.delegate = nullptr;
 		[Session release];
@@ -383,13 +181,43 @@ void FAppleARKitSystem::Shutdown()
 	CameraImage = nullptr;
 }
 
+void FAppleARKitSystem::CheckForFaceARSupport(UARSessionConfig* InSessionConfig)
+{
+	if (InSessionConfig->GetSessionType() != EARSessionType::Face)
+	{
+		// Clear the face ar support so we don't forward to it
+		FaceARSupport = TSharedPtr<FAppleARKitFaceSupportBase, ESPMode::ThreadSafe>();
+		return;
+	}
+
+	// We need to get the face support from the factory method, which is a modular feature to avoid dependencies
+	IModularFeature* Feature = IModularFeatures::Get().GetModularFeatureImplementation("AppleARKitFaceSupportFactory", 0);
+	if (ensureAlwaysMsgf(Feature != nullptr, TEXT("Face AR session has been requested but the face ar plugin is not enabled")))
+	{
+		IAppleARKitFaceSupportFactory* Factory = static_cast<IAppleARKitFaceSupportFactory*>(Feature);
+		FaceARSupport = Factory->CreateFaceSupport(SharedThis(this), this);
+		ensureAlwaysMsgf(FaceARSupport.IsValid(), TEXT("Face AR session has been requested but the face ar plugin is not enabled"));
+	}
+}
+
+UARTrackedGeometry* FAppleARKitSystem::GetTrackedGeometry(const FGuid& GeoGuid)
+{
+	UARTrackedGeometry** GeometrySearchResult = TrackedGeometries.Find(GeoGuid);
+	return GeometrySearchResult != nullptr ? *GeometrySearchResult : nullptr;
+}
+
+void FAppleARKitSystem::AddTrackedGeometry(const FGuid& Guid, UARTrackedGeometry* TrackedGeo)
+{
+	check(nullptr);
+	TrackedGeometries.Add(Guid, TrackedGeo);
+}
+
 FName FAppleARKitSystem::GetSystemName() const
 {
 	static const FName AppleARKitSystemName(TEXT("AppleARKit"));
 	return AppleARKitSystemName;
 }
 
-
 bool FAppleARKitSystem::GetCurrentPose(int32 DeviceId, FQuat& OutOrientation, FVector& OutPosition)
 {
 	if (DeviceId == IXRTrackingSystem::HMDDeviceId && GameThreadFrame.IsValid())
@@ -680,7 +508,7 @@ static UARTrackedGeometry* FindGeometryFromAnchor( ARAnchor* InAnchor, TMap<FGui
 {
 	if (InAnchor != NULL)
 	{
-		const FGuid AnchorGUID = ToFGuid( InAnchor.identifier );
+		const FGuid AnchorGUID = FAppleARKitConversion::ToFGuid( InAnchor.identifier );
 		UARTrackedGeometry** Result = Geometries.Find(AnchorGUID);
 		if (Result != nullptr)
 		{
@@ -748,7 +576,7 @@ TArray<FARTraceResult> FAppleARKitSystem::OnLineTraceTrackedObjects( const FVect
 						{
 							// Hit result has passed and above filtering, add it to the list
 							// Convert to Unreal's Hit Test result format
-							Results.Add( FARTraceResult( This, UnrealHitDistance, EARLineTraceChannels::PlaneUsingExtent, FAppleARKitTransform::ToFTransform( HitTestResult.worldTransform, WorldToMetersScale )*GetAlignmentTransform(), FindGeometryFromAnchor(HitTestResult.anchor, TrackedGeometries) ) );
+							Results.Add( FARTraceResult( This, UnrealHitDistance, EARLineTraceChannels::PlaneUsingExtent, FAppleARKitConversion::ToFTransform( HitTestResult.worldTransform )*GetAlignmentTransform(), FindGeometryFromAnchor(HitTestResult.anchor, TrackedGeometries) ) );
 						}
 					}
 				}
@@ -764,7 +592,7 @@ TArray<FARTraceResult> FAppleARKitSystem::OnLineTraceTrackedObjects( const FVect
 						{
 							// Hit result has passed and above filtering, add it to the list
 							// Convert to Unreal's Hit Test result format
-							Results.Add( FARTraceResult( This, UnrealHitDistance, EARLineTraceChannels::GroundPlane, FAppleARKitTransform::ToFTransform( HitTestResult.worldTransform, WorldToMetersScale )*GetAlignmentTransform(), FindGeometryFromAnchor(HitTestResult.anchor, TrackedGeometries) ) );
+							Results.Add( FARTraceResult( This, UnrealHitDistance, EARLineTraceChannels::GroundPlane, FAppleARKitConversion::ToFTransform( HitTestResult.worldTransform )*GetAlignmentTransform(), FindGeometryFromAnchor(HitTestResult.anchor, TrackedGeometries) ) );
 						}
 					}
 				}
@@ -782,7 +610,7 @@ TArray<FARTraceResult> FAppleARKitSystem::OnLineTraceTrackedObjects( const FVect
 						{
 							// Hit result has passed and above filtering, add it to the list
 							// Convert to Unreal's Hit Test result format
-							Results.Add( FARTraceResult( This, UnrealHitDistance, EARLineTraceChannels::FeaturePoint, FAppleARKitTransform::ToFTransform( HitTestResult.worldTransform, WorldToMetersScale )*GetAlignmentTransform(), FindGeometryFromAnchor(HitTestResult.anchor, TrackedGeometries) ) );
+							Results.Add( FARTraceResult( This, UnrealHitDistance, EARLineTraceChannels::FeaturePoint, FAppleARKitConversion::ToFTransform( HitTestResult.worldTransform )*GetAlignmentTransform(), FindGeometryFromAnchor(HitTestResult.anchor, TrackedGeometries) ) );
 						}
 					}
 				}
@@ -1017,8 +845,17 @@ bool FAppleARKitSystem::Run(UARSessionConfig* SessionConfig)
 	{
 		ARSessionRunOptions options = 0;
 
-		// Convert to native ARWorldTrackingSessionConfiguration
-		ARConfiguration* Configuration = ToARConfiguration(SessionConfig, Config, CandidateImages, ConvertedCandidateImages);
+		ARConfiguration* Configuration = nullptr;
+		CheckForFaceARSupport(SessionConfig);
+		if (!FaceARSupport.IsValid())
+		{
+			Configuration = ToARConfiguration(SessionConfig, Config, CandidateImages, ConvertedCandidateImages);
+		}
+		else
+		{
+			Configuration = FaceARSupport->ToARConfiguration(SessionConfig, Config);
+		}
+
 		// Not all session types are supported by all devices
 		if (Configuration == nullptr)
 		{
@@ -1141,18 +978,18 @@ void FAppleARKitSystem::SessionDidFailWithError_DelegateThread(const FString& Er
 
 #if SUPPORTS_ARKIT_1_0
 
-static TSharedPtr<FAppleARKitAnchorData> MakeAnchorData( ARAnchor* Anchor, const float WorldToMeterScale )
+static TSharedPtr<FAppleARKitAnchorData> MakeAnchorData( ARAnchor* Anchor )
 {
 	TSharedPtr<FAppleARKitAnchorData> NewAnchor;
 	if ([Anchor isKindOfClass:[ARPlaneAnchor class]])
 	{
 		ARPlaneAnchor* PlaneAnchor = (ARPlaneAnchor*)Anchor;
 		NewAnchor = MakeShared<FAppleARKitAnchorData>(
-			ToFGuid(PlaneAnchor.identifier),
-			FAppleARKitTransform::ToFTransform(PlaneAnchor.transform, WorldToMeterScale),
-			FAppleARKitTransform::ToFVector(PlaneAnchor.center, WorldToMeterScale),
+			FAppleARKitConversion::ToFGuid(PlaneAnchor.identifier),
+			FAppleARKitConversion::ToFTransform(PlaneAnchor.transform),
+			FAppleARKitConversion::ToFVector(PlaneAnchor.center),
 			// @todo use World Settings WorldToMetersScale
-			0.5f*FAppleARKitTransform::ToFVector(PlaneAnchor.extent, WorldToMeterScale).GetAbs()
+			0.5f*FAppleARKitConversion::ToFVector(PlaneAnchor.extent).GetAbs()
 		);
 
 #if SUPPORTS_ARKIT_1_5
@@ -1164,7 +1001,7 @@ static TSharedPtr<FAppleARKitAnchorData> MakeAnchorData( ARAnchor* Anchor, const
 			for(int32 i=0; i<NumBoundaryVerts; ++i)
 			{
 				const vector_float3& Vert = PlaneAnchor.geometry.boundaryVertices[i];
-				NewAnchor->BoundaryVerts.Add(FAppleARKitTransform::ToFVector(Vert, WorldToMeterScale));
+				NewAnchor->BoundaryVerts.Add(FAppleARKitConversion::ToFVector(Vert));
 			}
 		}
 #endif
@@ -1174,32 +1011,17 @@ static TSharedPtr<FAppleARKitAnchorData> MakeAnchorData( ARAnchor* Anchor, const
 	{
 		ARImageAnchor* ImageAnchor = (ARImageAnchor*)Anchor;
 		NewAnchor = MakeShared<FAppleARKitAnchorData>(
-			ToFGuid(ImageAnchor.identifier),
-			FAppleARKitTransform::ToFTransform(ImageAnchor.transform, WorldToMeterScale),
+			FAppleARKitConversion::ToFGuid(ImageAnchor.identifier),
+			FAppleARKitConversion::ToFTransform(ImageAnchor.transform),
 			FString(ImageAnchor.referenceImage.name)
 		);
 	}
 #endif
-	else if ([Anchor isKindOfClass:[ARFaceAnchor class]])
-	{
-		ARFaceAnchor* FaceAnchor = (ARFaceAnchor*)Anchor;
-		NewAnchor = MakeShared<FAppleARKitAnchorData>(
-			ToFGuid(FaceAnchor.identifier),
-			FAppleARKitTransform::ToFTransform(FaceAnchor.transform, WorldToMeterScale),
-			ToBlendShapeMap(FaceAnchor.blendShapes, FAppleARKitTransform::ToFTransform(FaceAnchor.transform)),
-			ToVertexBuffer(FaceAnchor.geometry.vertices, FaceAnchor.geometry.vertexCount)
-		);
-		// Only convert from 16bit to 32bit once
-		if (FAppleARKitAnchorData::FaceIndices.Num() == 0)
-		{
-			FAppleARKitAnchorData::FaceIndices = To32BitIndexBuffer(FaceAnchor.geometry.triangleIndices, FaceAnchor.geometry.triangleCount * 3);
-		}
-	}
 	else
 	{
 		NewAnchor = MakeShared<FAppleARKitAnchorData>(
-			ToFGuid(Anchor.identifier),
-			FAppleARKitTransform::ToFTransform(Anchor.transform, WorldToMeterScale));
+			FAppleARKitConversion::ToFGuid(Anchor.identifier),
+			FAppleARKitConversion::ToFTransform(Anchor.transform));
 	}
 	
 	return NewAnchor;
@@ -1211,9 +1033,16 @@ void FAppleARKitSystem::SessionDidAddAnchors_DelegateThread( NSArray<ARAnchor*>*
 					   STAT_FAppleARKitSystem_SessionDidAddAnchors,
 					   STATGROUP_APPLEARKIT);
 
+	// If this object is valid, we are running a face session and need that code to process things
+	if (FaceARSupport.IsValid())
+	{
+		FaceARSupport->ProcessAnchorAdd(anchors, GetAlignmentTransform(), GameThreadFrameNumber, GameThreadTimestamp);
+		return;
+	}
+
 	for (ARAnchor* anchor in anchors)
 	{
-		TSharedPtr<FAppleARKitAnchorData> NewAnchorData = MakeAnchorData(anchor, GetWorldToMetersScale());
+		TSharedPtr<FAppleARKitAnchorData> NewAnchorData = MakeAnchorData(anchor);
 		if (ensure(NewAnchorData.IsValid()))
 		{
 			auto AddAnchorTask = FSimpleDelegateGraphTask::FDelegate::CreateSP(this, &FAppleARKitSystem::SessionDidAddAnchors_Internal, NewAnchorData.ToSharedRef());
@@ -1228,9 +1057,16 @@ void FAppleARKitSystem::SessionDidUpdateAnchors_DelegateThread( NSArray<ARAnchor
 					   STAT_FAppleARKitSystem_SessionDidUpdateAnchors,
 					   STATGROUP_APPLEARKIT);
 	
+	// If this object is valid, we are running a face session and need that code to process things
+	if (FaceARSupport.IsValid())
+	{
+		FaceARSupport->ProcessAnchorUpdate(anchors, GetAlignmentTransform(), GameThreadFrameNumber, GameThreadTimestamp);
+		return;
+	}
+
 	for (ARAnchor* anchor in anchors)
 	{
-		TSharedPtr<FAppleARKitAnchorData> NewAnchorData = MakeAnchorData(anchor, GetWorldToMetersScale());
+		TSharedPtr<FAppleARKitAnchorData> NewAnchorData = MakeAnchorData(anchor);
 		if (ensure(NewAnchorData.IsValid()))
 		{
 			auto UpdateAnchorTask = FSimpleDelegateGraphTask::FDelegate::CreateSP(this, &FAppleARKitSystem::SessionDidUpdateAnchors_Internal, NewAnchorData.ToSharedRef());
@@ -1245,10 +1081,12 @@ void FAppleARKitSystem::SessionDidRemoveAnchors_DelegateThread( NSArray<ARAnchor
 					   STAT_FAppleARKitSystem_SessionDidRemoveAnchors,
 					   STATGROUP_APPLEARKIT);
 	
+	// Face AR Anchors are also removed this way, no need for special code since they are tracked geometry
+
 	for (ARAnchor* anchor in anchors)
 	{
 		// Convert to FGuid
-		const FGuid AnchorGuid = ToFGuid( anchor.identifier );
+		const FGuid AnchorGuid = FAppleARKitConversion::ToFGuid( anchor.identifier );
 
 		auto RemoveAnchorTask = FSimpleDelegateGraphTask::FDelegate::CreateSP(this, &FAppleARKitSystem::SessionDidRemoveAnchors_Internal, AnchorGuid);
 		FSimpleDelegateGraphTask::CreateAndDispatchWhenReady(RemoveAnchorTask, GET_STATID(STAT_FAppleARKitSystem_SessionDidRemoveAnchors), nullptr, ENamedThreads::GameThread);
@@ -1283,20 +1121,6 @@ void FAppleARKitSystem::SessionDidAddAnchors_Internal( TSharedRef<FAppleARKitAnc
 			NewGeometry = NewGeo;
 			break;
 		}
-		case EAppleAnchorType::FaceAnchor:
-		{
-			NewAnchorDebugName = FString::Printf(TEXT("FACE-%02d"), LastTrackedGeometry_DebugId++);
-			// Update LiveLink first, because the other updates use MoveTemp for efficiency
-			if (LiveLinkSource.IsValid())
-			{
-				LiveLinkSource->PublishBlendShapes(FaceTrackingLiveLinkSubjectName, GameThreadTimestamp, GameThreadFrameNumber, AnchorData->BlendShapes);
-			}
-
-			UARFaceGeometry* NewGeo = NewObject<UARFaceGeometry>();
-			NewGeo->UpdateTrackedGeometry(SharedThis(this), GameThreadFrameNumber, GameThreadTimestamp, AnchorData->Transform, GetAlignmentTransform(), AnchorData->BlendShapes, AnchorData->FaceVerts, AnchorData->FaceIndices);
-			NewGeometry = NewGeo;
-			break;
-		}
 		case EAppleAnchorType::ImageAnchor:
 		{
 			NewAnchorDebugName = FString::Printf(TEXT("IMG-%02d"), LastTrackedGeometry_DebugId++);
@@ -1368,25 +1192,6 @@ void FAppleARKitSystem::SessionDidUpdateAnchors_Internal( TSharedRef<FAppleARKit
 				}
 				break;
 			}
-			case EAppleAnchorType::FaceAnchor:
-			{
-				if (UARFaceGeometry* FaceGeo = Cast<UARFaceGeometry>(FoundGeometry))
-				{
-					// Update LiveLink first, because the other updates use MoveTemp for efficiency
-					if (LiveLinkSource.IsValid())
-					{
-						LiveLinkSource->PublishBlendShapes(FaceTrackingLiveLinkSubjectName, GameThreadTimestamp, GameThreadFrameNumber, AnchorData->BlendShapes);
-					}
-
-					FaceGeo->UpdateTrackedGeometry(SharedThis(this), GameThreadFrameNumber, GameThreadTimestamp, AnchorData->Transform, GetAlignmentTransform(), AnchorData->BlendShapes, AnchorData->FaceVerts, AnchorData->FaceIndices);
-					for (UARPin* Pin : PinsToUpdate)
-					{
-						const FTransform Pin_LocalToTrackingTransform_PostUpdate = Pin->GetLocalToTrackingTransform_NoAlignment() * AnchorDeltaTransform;
-						Pin->OnTransformUpdated(Pin_LocalToTrackingTransform_PostUpdate);
-					}
-				}
-				break;
-			}
 		}
 	}
 }
diff --git a/Engine/Plugins/Runtime/AppleARKit/Source/AppleARKit/Private/AppleARKitSystem.h b/Engine/Plugins/Runtime/AppleARKit/Source/AppleARKit/Private/AppleARKitSystem.h
--- a/Engine/Plugins/Runtime/AppleARKit/Source/AppleARKit/Private/AppleARKitSystem.h
+++ b/Engine/Plugins/Runtime/AppleARKit/Source/AppleARKit/Private/AppleARKitSystem.h
@@ -6,9 +6,9 @@
 #include "AppleARKitConfiguration.h"
 #include "ARSystem.h"
 #include "AppleARKitHitTestResult.h"
-#include "AppleARKitLiveLinkSourceFactory.h"
 #include "AppleARKitTextures.h"
 #include "Kismet/BlueprintPlatformLibrary.h"
+#include "AppleARKitFaceSupport.h"
 
 // ARKit
 #if SUPPORTS_ARKIT_1_0
@@ -26,7 +26,9 @@ DECLARE_STATS_GROUP(TEXT("AppleARKit"), STATGROUP_APPLEARKIT, STATCAT_Advanced);
 struct FAppleARKitFrame;
 struct FAppleARKitAnchorData;
 
-class FAppleARKitSystem : public FARSystemBase
+class FAppleARKitSystem :
+	public FARSystemBase,
+	public IAppleARKitFaceSupportCallback
 {
 	friend class FAppleARKitXRCamera;
 	
@@ -48,6 +50,11 @@ public:
 	bool OnStartGameFrame(FWorldContext& WorldContext) override;
 	//~ IXRTrackingSystem
 
+	//~IAppleARKitFaceSupportCallback
+	virtual UARTrackedGeometry* GetTrackedGeometry(const FGuid& GeoGuid) override;
+	virtual void AddTrackedGeometry(const FGuid& Guid, UARTrackedGeometry* TrackedGeo) override;
+	//~IAppleARKitFaceSupportCallback
+
 	void* GetARSessionRawPointer() override;
 	void* GetGameThreadARFrameRawPointer() override;
 	
@@ -126,6 +133,9 @@ private:
 	bool bIsRunning = false;
 	
 	void SetDeviceOrientation( EScreenOrientation::Type InOrientation );
+
+	/** Creates or clears the face ar support object if face ar has been requested */
+	void CheckForFaceARSupport(UARSessionConfig* InSessionConfig);
 	
 	/** The orientation of the device; see EScreenOrientation */
 	EScreenOrientation::Type DeviceOrientation;
@@ -172,12 +182,6 @@ private:
 	/** The ar timestamp of when the LastReceivedFrame was last updated */
 	double GameThreadTimestamp;
 
-	/** If requested, publishes face ar updates to LiveLink for the animation system to use */
-	TSharedPtr<ILiveLinkSourceARKit> LiveLinkSource;
-	/** Copied from the UARSessionConfig project settings object */
-	FName FaceTrackingLiveLinkSubjectName;
-	
-	
 	// An int counter that provides a human-readable debug number for Tracked Geometries.
 	uint32 LastTrackedGeometry_DebugId;
 
@@ -189,6 +193,9 @@ private:
 	TSharedPtr< FAppleARKitFrame, ESPMode::ThreadSafe > GameThreadFrame;
 	TSharedPtr< FAppleARKitFrame, ESPMode::ThreadSafe > RenderThreadFrame;
 	TSharedPtr< FAppleARKitFrame, ESPMode::ThreadSafe > LastReceivedFrame;
+
+	// The object that is handling face support if present
+	TSharedPtr<FAppleARKitFaceSupportBase, ESPMode::ThreadSafe> FaceARSupport;
 };
 
 
diff --git a/Engine/Plugins/Runtime/AppleARKit/Source/AppleARKit/Public/AppleARKitConversion.h b/Engine/Plugins/Runtime/AppleARKit/Source/AppleARKit/Public/AppleARKitConversion.h
--- a/Engine/Plugins/Runtime/AppleARKit/Source/AppleARKit/Public/AppleARKitConversion.h
+++ b/Engine/Plugins/Runtime/AppleARKit/Source/AppleARKit/Public/AppleARKitConversion.h
@@ -4,19 +4,48 @@
 
 // UE4
 #include "Math/Transform.h"
+#include "ARPin.h"
 
-struct FAppleARKitTransform
+//@todo -- Move this to UARSessionConfiguration
+/**
+ * Enum constants for indicating the world alignment.
+ */
+enum class EAppleARKitWorldAlignment : uint8
 {
+	/** Aligns the world with gravity that is defined by vector (0, -1, 0). */
+	Gravity,
 
-#if SUPPORTS_ARKIT_1_0
+	/**
+	 * Aligns the world with gravity that is defined by the vector (0, -1, 0)
+	 * and heading (w.r.t. True North) that is given by the vector (0, 0, -1).
+	 */
+	GravityAndHeading,
 
-	/** 
+	/** Aligns the world with the camera's orientation. */
+	Camera
+};
+ENUM_CLASS_FLAGS(EAppleARKitWorldAlignment)
+
+struct FAppleARKitConversion
+{
+	static FORCEINLINE float ToUE4Scale()
+	{
+		return 100.f;
+	}
+
+	static FORCEINLINE float ToARKitScale()
+	{
+		return 0.01f;
+	}
+
+#if SUPPORTS_ARKIT_1_0
+	/**
 	 * Convert's an ARKit 'Y up' 'right handed' coordinate system transform to Unreal's 'Z up' 
 	 * 'left handed' coordinate system.
 	 *
 	 * Ignores scale.
 	 */
-	static FORCEINLINE FTransform ToFTransform( const matrix_float4x4& RawYUpMatrix, float WorldToMetersScale = 100.0f )
+	static FORCEINLINE FTransform ToFTransform(const matrix_float4x4& RawYUpMatrix)
 	{
 		// Conversion here is as per SteamVRHMD::ToFMatrix
 		FMatrix RawYUpFMatrix(
@@ -26,7 +55,7 @@ struct FAppleARKitTransform
 			FPlane(RawYUpMatrix.columns[3][0], RawYUpMatrix.columns[3][1], RawYUpMatrix.columns[3][2], RawYUpMatrix.columns[3][3]));
 
 		// Extract & convert translation
-		FVector Translation = FVector( -RawYUpFMatrix.M[3][2], RawYUpFMatrix.M[3][0], RawYUpFMatrix.M[3][1] ) * WorldToMetersScale;
+		FVector Translation = FVector( -RawYUpFMatrix.M[3][2], RawYUpFMatrix.M[3][0], RawYUpFMatrix.M[3][1] ) * ToUE4Scale();
 
 		// Extract & convert rotation 
 		FQuat RawRotation( RawYUpFMatrix );
@@ -39,11 +68,45 @@ struct FAppleARKitTransform
 	 * Convert's an ARKit 'Y up' 'right handed' coordinate system vector to Unreal's 'Z up' 
 	 * 'left handed' coordinate system.
 	 */
-	static FORCEINLINE FVector ToFVector( const vector_float3& RawYUpVector, float WorldToMetersScale = 100.0f )
+	static FORCEINLINE FVector ToFVector(const vector_float3& RawYUpVector)
+	{
+		return FVector( -RawYUpVector.z, RawYUpVector.x, RawYUpVector.y ) * ToUE4Scale();
+	}
+
+	static FORCEINLINE FGuid ToFGuid( uuid_t UUID )
+	{
+		FGuid AsGUID(
+			*(uint32*)UUID,
+			*((uint32*)UUID)+1,
+			*((uint32*)UUID)+2,
+			*((uint32*)UUID)+3);
+		return AsGUID;
+	}
+
+	static FORCEINLINE FGuid ToFGuid( NSUUID* Identifier )
 	{
-		return FVector( -RawYUpVector.z, RawYUpVector.x, RawYUpVector.y ) * WorldToMetersScale;
+		// Get bytes
+		uuid_t UUID;
+		[Identifier getUUIDBytes:UUID];
+
+		// Set FGuid parts
+		return ToFGuid( UUID );
 	}
 
+	static FORCEINLINE ARWorldAlignment ToARWorldAlignment( const EAppleARKitWorldAlignment& InWorldAlignment )
+	{
+		switch ( InWorldAlignment )
+		{
+			case EAppleARKitWorldAlignment::Gravity:
+				return ARWorldAlignmentGravity;
+
+			case EAppleARKitWorldAlignment::GravityAndHeading:
+				return ARWorldAlignmentGravityAndHeading;
+
+			case EAppleARKitWorldAlignment::Camera:
+				return ARWorldAlignmentCamera;
+		};
+	};
 #endif
 };
 
@@ -74,15 +137,6 @@ struct FAppleARKitAnchorData
 	{
 	}
 
-	FAppleARKitAnchorData(FGuid InAnchorGuid, FTransform InTransform, FARBlendShapeMap InBlendShapes, TArray<FVector> InFaceVerts)
-		: Transform( InTransform )
-		, AnchorType( EAppleAnchorType::FaceAnchor )
-		, AnchorGUID( InAnchorGuid )
-		, BlendShapes( MoveTemp(InBlendShapes) )
-		, FaceVerts( MoveTemp(InFaceVerts) )
-	{
-	}
-
 	FAppleARKitAnchorData(FGuid InAnchorGuid, FTransform InTransform, FString InDetectedImageName)
 		: Transform( InTransform )
 		, AnchorType( EAppleAnchorType::ImageAnchor )
@@ -96,15 +150,37 @@ struct FAppleARKitAnchorData
 	FGuid AnchorGUID;
 	FVector Center;
 	FVector Extent;
-
 	TArray<FVector> BoundaryVerts;
-	FARBlendShapeMap BlendShapes;
-	TArray<FVector> FaceVerts;
-	// Temp non-static while code is rearranged
-	TArray<int32> FaceIndices;
-	// Note: the index buffer never changes so can be safely read once
-//	static TArray<int32> FaceIndices;
-//	TArray<int32> FAppleARKitAnchorData::FaceIndices;
 
 	FString DetectedImageName;
 };
+
+namespace ARKitUtil
+{
+	static UARPin* PinFromComponent( const USceneComponent* Component, const TArray<UARPin*>& InPins )
+	{
+		for (UARPin* Pin : InPins)
+		{
+			if (Pin->GetPinnedComponent() == Component)
+			{
+				return Pin;
+			}
+		}
+
+		return nullptr;
+	}
+
+	static TArray<UARPin*> PinsFromGeometry( const UARTrackedGeometry* Geometry, const TArray<UARPin*>& InPins )
+	{
+		TArray<UARPin*> OutPins;
+		for (UARPin* Pin : InPins)
+		{
+			if (Pin->GetTrackedGeometry() == Geometry)
+			{
+				OutPins.Add(Pin);
+			}
+		}
+
+		return OutPins;
+	}
+}
diff --git a/Engine/Plugins/Runtime/AppleARKit/Source/AppleARKit/Public/AppleARKitFaceMeshComponent.h b/Engine/Plugins/Runtime/AppleARKit/Source/AppleARKit/Public/AppleARKitFaceMeshComponent.h
deleted file mode 100644
diff --git a/Engine/Plugins/Runtime/AppleARKit/Source/AppleARKit/Public/AppleARKitFaceSupport.h b/Engine/Plugins/Runtime/AppleARKit/Source/AppleARKit/Public/AppleARKitFaceSupport.h
new file mode 100644
--- /dev/null
+++ b/Engine/Plugins/Runtime/AppleARKit/Source/AppleARKit/Public/AppleARKitFaceSupport.h
@@ -0,0 +1,128 @@
+// Copyright 1998-2018 Epic Games, Inc. All Rights Reserved.
+
+#pragma once
+
+#include "Features/IModularFeature.h"
+#include "Features/IModularFeatures.h"
+
+#include "AppleARKitAvailability.h"
+
+#if SUPPORTS_ARKIT_1_0
+	#import <ARKit/ARKit.h>
+#endif
+
+class UARSessionConfig;
+class UARTrackedGeometry;
+class FAppleARKitConfiguration;
+class FARSystemBase;
+
+
+class APPLEARKIT_API IAppleARKitFaceSupportCallback
+{
+public:
+	/** So that face processing can get access to the face geometry objects by their guid */
+	virtual UARTrackedGeometry* GetTrackedGeometry(const FGuid& GeoGuid) = 0;
+	/** So that face processing can add new geometry as updates come in */
+	virtual void AddTrackedGeometry(const FGuid& Guid, UARTrackedGeometry* TrackedGeo) = 0;
+};
+
+class APPLEARKIT_API IAppleARKitFaceSupport
+{
+public:
+#if SUPPORTS_ARKIT_1_0
+	/**
+	 * Forwards the anchor add for face processing
+	 *
+	 * @param NewAnchors the anchor list that the face support should process
+	 * @param AlignmentTransform the transform to apply for user alignment
+	 * @param FrameNumber the frame number to publish with LiveLink
+	 * @param Timestamp the timestamp to publish with LiveLink
+	 */
+	virtual void ProcessAnchorAdd(NSArray<ARAnchor*>* NewAnchors, const FTransform& AlignmentTransform, uint32 FrameNumber, double Timestamp) = 0;
+
+	/**
+	 * Forwards the anchor add for face processing
+	 *
+	 * @param UpdatedAnchors the anchor list that the face support should process
+	 * @param AlignmentTransform the transform to apply for user alignment
+	 * @param FrameNumber the frame number to publish with LiveLink
+	 * @param Timestamp the timestamp to publish with LiveLink
+	 */
+	virtual void ProcessAnchorUpdate(NSArray<ARAnchor*>* UpdatedAnchors, const FTransform& AlignmentTransform, uint32 FrameNumber, double Timestamp) = 0;
+
+	/**
+	 * Creates a face ar specific configuration object if that is requested
+	 *
+	 * @param SessionConfig the UE4 configuration object that needs processing
+	 * @param InConfiguration the legacy configuration object that needs to go away
+	 */
+	virtual ARConfiguration* ToARConfiguration(UARSessionConfig* SessionConfig, FAppleARKitConfiguration& InConfiguration) = 0;
+#endif
+
+	IAppleARKitFaceSupport() { }
+	IAppleARKitFaceSupport(const TSharedRef<FARSystemBase, ESPMode::ThreadSafe>& InTrackingSystem, IAppleARKitFaceSupportCallback* Callback) { }
+	virtual ~IAppleARKitFaceSupport() { }
+};
+
+class APPLEARKIT_API FAppleARKitFaceSupportBase :
+	public IAppleARKitFaceSupport,
+	public TSharedFromThis<FAppleARKitFaceSupportBase, ESPMode::ThreadSafe>
+{
+public:
+#if SUPPORTS_ARKIT_1_0
+	/**
+	 * Forwards the anchor add for face processing
+	 *
+	 * @param NewAnchors the anchor list that the face support should process
+	 * @param AlignmentTransform the transform to apply for user alignment
+	 * @param FrameNumber the frame number to publish with LiveLink
+	 * @param Timestamp the timestamp to publish with LiveLink
+	 */
+	virtual void ProcessAnchorAdd(NSArray<ARAnchor*>* NewAnchors, const FTransform& AlignmentTransform, uint32 FrameNumber, double Timestamp) override
+	{
+
+	}
+
+	/**
+	 * Forwards the anchor add for face processing
+	 *
+	 * @param UpdatedAnchors the anchor list that the face support should process
+	 * @param AlignmentTransform the transform to apply for user alignment
+	 * @param FrameNumber the frame number to publish with LiveLink
+	 * @param Timestamp the timestamp to publish with LiveLink
+	 */
+	virtual void ProcessAnchorUpdate(NSArray<ARAnchor*>* UpdatedAnchors, const FTransform& AlignmentTransform, uint32 FrameNumber, double Timestamp) override
+	{
+
+	}
+
+	/**
+	 * Creates a face ar specific configuration object if that is requested
+	 *
+	 * @param SessionConfig the UE4 configuration object that needs processing
+	 * @param InConfiguration the legacy configuration object that needs to go away
+	 */
+	virtual ARConfiguration* ToARConfiguration(UARSessionConfig* SessionConfig, FAppleARKitConfiguration& InConfiguration) override
+	{
+		return nullptr;
+	}
+#endif
+
+	FAppleARKitFaceSupportBase() { }
+	FAppleARKitFaceSupportBase(const TSharedRef<FARSystemBase, ESPMode::ThreadSafe>& InTrackingSystem, IAppleARKitFaceSupportCallback* Callback) { }
+	virtual ~FAppleARKitFaceSupportBase() { }
+};
+
+class APPLEARKIT_API IAppleARKitFaceSupportFactory :
+	public IModularFeature
+{
+public:
+	/** Factory method that returns the object to use to handle face ar requests */
+	virtual TSharedPtr<FAppleARKitFaceSupportBase, ESPMode::ThreadSafe> CreateFaceSupport(const TSharedRef<FARSystemBase, ESPMode::ThreadSafe>& InTrackingSystem, IAppleARKitFaceSupportCallback* Callback) = 0;
+
+	static FName GetModularFeatureName()
+	{
+		static FName FeatureName = FName(TEXT("AppleARKitFaceSupportFactory"));
+		return FeatureName;
+	}
+};
diff --git a/Engine/Plugins/Runtime/AppleARKit/Source/AppleARKit/Public/AppleARKitLiveLinkSourceFactory.h b/Engine/Plugins/Runtime/AppleARKit/Source/AppleARKit/Public/AppleARKitLiveLinkSourceFactory.h
deleted file mode 100644
diff --git a/Engine/Plugins/Runtime/AppleARKitFaceSupport/AppleARKitFaceSupport.uplugin b/Engine/Plugins/Runtime/AppleARKitFaceSupport/AppleARKitFaceSupport.uplugin
new file mode 100644
--- /dev/null
+++ b/Engine/Plugins/Runtime/AppleARKitFaceSupport/AppleARKitFaceSupport.uplugin
@@ -0,0 +1,44 @@
+{
+	"FileVersion": 3,
+	"Version": 1,
+	"VersionName": "0.1",
+	"FriendlyName": "Apple ARKit Face Support",
+	"Description": "Support for Apple's face tracking features",
+	"Category": "Augmented Reality",
+	"CreatedBy": "",
+	"CreatedByURL": "",
+	"DocsURL": "",
+	"MarketplaceURL": "",
+	"SupportURL": "",
+	"CanContainContent": true,
+	"IsBetaVersion": true,
+	"EnabledByDefault": false,
+	"Modules":
+	[
+		{
+			"Name": "AppleARKitFaceSupport",
+			"Type": "Runtime",
+			"LoadingPhase": "PostConfigInit",
+			"WhitelistPlatforms" : [ "Mac", "iOS", "Win64" ]
+		}
+	],
+    "Plugins":
+    [
+		{
+			"Name": "AppleARKit",
+			"Enabled": true
+		},
+		{
+			"Name": "ProceduralMeshComponent",
+			"Enabled": true
+		},
+        {
+            "Name": "LiveLink",
+            "Enabled": true
+        },
+        {
+            "Name": "AppleImageUtils",
+            "Enabled": true
+        }
+    ]
+}
diff --git a/Engine/Plugins/Runtime/AppleARKitFaceSupport/Source/AppleARKitFaceSupport/AppleARKitFaceSupport.Build.cs b/Engine/Plugins/Runtime/AppleARKitFaceSupport/Source/AppleARKitFaceSupport/AppleARKitFaceSupport.Build.cs
new file mode 100644
--- /dev/null
+++ b/Engine/Plugins/Runtime/AppleARKitFaceSupport/Source/AppleARKitFaceSupport/AppleARKitFaceSupport.Build.cs
@@ -0,0 +1,58 @@
+// Copyright 1998-2018 Epic Games, Inc. All Rights Reserved.
+
+using UnrealBuildTool;
+
+public class AppleARKitFaceSupport : ModuleRules
+{
+	public AppleARKitFaceSupport(ReadOnlyTargetRules Target) : base(Target)
+	{
+		PrivateIncludePaths.AddRange(
+			new string[] {
+				"../../../../Plugins/Runtime/AppleARKit/Source/AppleARKit/Private",
+				// ... add other private include paths required here ...
+			}
+			);
+			
+		
+		PublicDependencyModuleNames.AddRange(
+			new string[]
+			{
+				"Core",
+				"Engine",
+				// ... add other public dependencies that you statically link with here ...
+			}
+			);
+			
+		
+		PrivateDependencyModuleNames.AddRange(
+			new string[]
+			{
+				"CoreUObject",
+                "IOSRuntimeSettings",
+                "HeadMountedDisplay",
+                "AugmentedReality",
+                "ProceduralMeshComponent",
+                "LiveLink",
+                "LiveLinkInterface",
+//                "OnlineSubsystem",
+                "Sockets",
+                "AppleARKit",
+                "AppleImageUtils"
+				// ... add private dependencies that you statically link with here ...
+			}
+			);
+		
+		
+		DynamicallyLoadedModuleNames.AddRange(
+			new string[]
+			{
+				// ... add any modules that your module loads dynamically here ...
+			}
+			);
+
+		if (Target.Platform == UnrealTargetPlatform.IOS)
+		{
+			PublicFrameworks.Add( "ARKit" );
+		}
+	}
+}
diff --git a/Engine/Plugins/Runtime/AppleARKitFaceSupport/Source/AppleARKitFaceSupport/Private/AppleARKitFaceMeshComponent.cpp b/Engine/Plugins/Runtime/AppleARKitFaceSupport/Source/AppleARKitFaceSupport/Private/AppleARKitFaceMeshComponent.cpp
new file mode 100644
--- /dev/null
+++ b/Engine/Plugins/Runtime/AppleARKitFaceSupport/Source/AppleARKitFaceSupport/Private/AppleARKitFaceMeshComponent.cpp
@@ -0,0 +1,363 @@
+// Copyright 1998-2018 Epic Games, Inc. All Rights Reserved.
+
+#include "AppleARKitFaceMeshComponent.h"
+#include "ARBlueprintLibrary.h"
+
+#if SUPPORTS_ARKIT_1_0
+	#import <ARKit/ARKit.h>
+#endif
+
+#include "AppleARKitFaceMeshConversion.h"
+#include "ARTrackable.h"
+#include "ARSessionConfig.h"
+#include "AppleARKitSettings.h"
+#include "AppleARKitConversion.h"
+
+
+#if SUPPORTS_ARKIT_1_0
+
+NSDictionary<ARBlendShapeLocation,NSNumber *>* ToBlendShapeDictionary(const FARBlendShapeMap& BlendShapeMap)
+{
+	NSMutableDictionary<ARBlendShapeLocation,NSNumber *>* BlendShapeDict = [[[NSMutableDictionary<ARBlendShapeLocation,NSNumber *> alloc] init] autorelease];
+
+#define SET_BLEND_SHAPE(AppleShape, UE4Shape) \
+	if (BlendShapeMap.Contains(UE4Shape)) \
+	{ \
+		NSNumber* Num = [NSNumber numberWithFloat: BlendShapeMap[UE4Shape]]; \
+		BlendShapeDict[AppleShape] = Num; \
+	}
+
+	// Do we want to capture face performance or look at the face as if in a mirror
+	if (GetDefault<UAppleARKitSettings>()->DefaultFaceTrackingDirection == EARFaceTrackingDirection::FaceRelative)
+	{
+		SET_BLEND_SHAPE(ARBlendShapeLocationEyeBlinkLeft, EARFaceBlendShape::EyeBlinkLeft);
+		SET_BLEND_SHAPE(ARBlendShapeLocationEyeLookDownLeft, EARFaceBlendShape::EyeLookDownLeft);
+		SET_BLEND_SHAPE(ARBlendShapeLocationEyeLookInLeft, EARFaceBlendShape::EyeLookInLeft);
+		SET_BLEND_SHAPE(ARBlendShapeLocationEyeLookOutLeft, EARFaceBlendShape::EyeLookOutLeft);
+		SET_BLEND_SHAPE(ARBlendShapeLocationEyeLookUpLeft, EARFaceBlendShape::EyeLookUpLeft);
+		SET_BLEND_SHAPE(ARBlendShapeLocationEyeSquintLeft, EARFaceBlendShape::EyeSquintLeft);
+		SET_BLEND_SHAPE(ARBlendShapeLocationEyeWideLeft, EARFaceBlendShape::EyeWideLeft);
+		SET_BLEND_SHAPE(ARBlendShapeLocationEyeBlinkRight, EARFaceBlendShape::EyeBlinkRight);
+		SET_BLEND_SHAPE(ARBlendShapeLocationEyeLookDownRight, EARFaceBlendShape::EyeLookDownRight);
+		SET_BLEND_SHAPE(ARBlendShapeLocationEyeLookInRight, EARFaceBlendShape::EyeLookInRight);
+		SET_BLEND_SHAPE(ARBlendShapeLocationEyeLookOutRight, EARFaceBlendShape::EyeLookOutRight);
+		SET_BLEND_SHAPE(ARBlendShapeLocationEyeLookUpRight, EARFaceBlendShape::EyeLookUpRight);
+		SET_BLEND_SHAPE(ARBlendShapeLocationEyeSquintRight, EARFaceBlendShape::EyeSquintRight);
+		SET_BLEND_SHAPE(ARBlendShapeLocationEyeWideRight, EARFaceBlendShape::EyeWideRight);
+		SET_BLEND_SHAPE(ARBlendShapeLocationJawForward, EARFaceBlendShape::JawForward);
+		SET_BLEND_SHAPE(ARBlendShapeLocationJawLeft, EARFaceBlendShape::JawLeft);
+		SET_BLEND_SHAPE(ARBlendShapeLocationJawRight, EARFaceBlendShape::JawRight);
+		SET_BLEND_SHAPE(ARBlendShapeLocationMouthLeft, EARFaceBlendShape::MouthLeft);
+		SET_BLEND_SHAPE(ARBlendShapeLocationMouthRight, EARFaceBlendShape::MouthRight);
+		SET_BLEND_SHAPE(ARBlendShapeLocationMouthSmileLeft, EARFaceBlendShape::MouthSmileLeft);
+		SET_BLEND_SHAPE(ARBlendShapeLocationMouthSmileRight, EARFaceBlendShape::MouthSmileRight);
+		SET_BLEND_SHAPE(ARBlendShapeLocationMouthFrownLeft, EARFaceBlendShape::MouthFrownLeft);
+		SET_BLEND_SHAPE(ARBlendShapeLocationMouthFrownRight, EARFaceBlendShape::MouthFrownRight);
+		SET_BLEND_SHAPE(ARBlendShapeLocationMouthDimpleLeft, EARFaceBlendShape::MouthDimpleLeft);
+		SET_BLEND_SHAPE(ARBlendShapeLocationMouthDimpleRight, EARFaceBlendShape::MouthDimpleRight);
+		SET_BLEND_SHAPE(ARBlendShapeLocationMouthStretchLeft, EARFaceBlendShape::MouthStretchLeft);
+		SET_BLEND_SHAPE(ARBlendShapeLocationMouthStretchRight, EARFaceBlendShape::MouthStretchRight);
+		SET_BLEND_SHAPE(ARBlendShapeLocationMouthPressLeft, EARFaceBlendShape::MouthPressLeft);
+		SET_BLEND_SHAPE(ARBlendShapeLocationMouthPressRight, EARFaceBlendShape::MouthPressRight);
+		SET_BLEND_SHAPE(ARBlendShapeLocationMouthLowerDownLeft, EARFaceBlendShape::MouthLowerDownLeft);
+		SET_BLEND_SHAPE(ARBlendShapeLocationMouthLowerDownRight, EARFaceBlendShape::MouthLowerDownRight);
+		SET_BLEND_SHAPE(ARBlendShapeLocationMouthUpperUpLeft, EARFaceBlendShape::MouthUpperUpLeft);
+		SET_BLEND_SHAPE(ARBlendShapeLocationMouthUpperUpRight, EARFaceBlendShape::MouthUpperUpRight);
+		SET_BLEND_SHAPE(ARBlendShapeLocationBrowDownLeft, EARFaceBlendShape::BrowDownLeft);
+		SET_BLEND_SHAPE(ARBlendShapeLocationBrowDownRight, EARFaceBlendShape::BrowDownRight);
+		SET_BLEND_SHAPE(ARBlendShapeLocationBrowOuterUpLeft, EARFaceBlendShape::BrowOuterUpLeft);
+		SET_BLEND_SHAPE(ARBlendShapeLocationBrowOuterUpRight, EARFaceBlendShape::BrowOuterUpRight);
+		SET_BLEND_SHAPE(ARBlendShapeLocationCheekSquintLeft, EARFaceBlendShape::CheekSquintLeft);
+		SET_BLEND_SHAPE(ARBlendShapeLocationCheekSquintRight, EARFaceBlendShape::CheekSquintRight);
+		SET_BLEND_SHAPE(ARBlendShapeLocationNoseSneerLeft, EARFaceBlendShape::NoseSneerLeft);
+		SET_BLEND_SHAPE(ARBlendShapeLocationNoseSneerRight, EARFaceBlendShape::NoseSneerRight);
+	}
+	else
+	{
+		SET_BLEND_SHAPE(ARBlendShapeLocationEyeBlinkLeft, EARFaceBlendShape::EyeBlinkRight);
+		SET_BLEND_SHAPE(ARBlendShapeLocationEyeLookDownLeft, EARFaceBlendShape::EyeLookDownRight);
+		SET_BLEND_SHAPE(ARBlendShapeLocationEyeLookInLeft, EARFaceBlendShape::EyeLookInRight);
+		SET_BLEND_SHAPE(ARBlendShapeLocationEyeLookOutLeft, EARFaceBlendShape::EyeLookOutRight);
+		SET_BLEND_SHAPE(ARBlendShapeLocationEyeLookUpLeft, EARFaceBlendShape::EyeLookUpRight);
+		SET_BLEND_SHAPE(ARBlendShapeLocationEyeSquintLeft, EARFaceBlendShape::EyeSquintRight);
+		SET_BLEND_SHAPE(ARBlendShapeLocationEyeWideLeft, EARFaceBlendShape::EyeWideRight);
+		SET_BLEND_SHAPE(ARBlendShapeLocationEyeBlinkRight, EARFaceBlendShape::EyeBlinkLeft);
+		SET_BLEND_SHAPE(ARBlendShapeLocationEyeLookDownRight, EARFaceBlendShape::EyeLookDownLeft);
+		SET_BLEND_SHAPE(ARBlendShapeLocationEyeLookInRight, EARFaceBlendShape::EyeLookInLeft);
+		SET_BLEND_SHAPE(ARBlendShapeLocationEyeLookOutRight, EARFaceBlendShape::EyeLookOutLeft);
+		SET_BLEND_SHAPE(ARBlendShapeLocationEyeLookUpRight, EARFaceBlendShape::EyeLookUpLeft);
+		SET_BLEND_SHAPE(ARBlendShapeLocationEyeSquintRight, EARFaceBlendShape::EyeSquintLeft);
+		SET_BLEND_SHAPE(ARBlendShapeLocationEyeWideRight, EARFaceBlendShape::EyeWideLeft);
+		SET_BLEND_SHAPE(ARBlendShapeLocationJawForward, EARFaceBlendShape::JawForward);
+		SET_BLEND_SHAPE(ARBlendShapeLocationJawLeft, EARFaceBlendShape::JawRight);
+		SET_BLEND_SHAPE(ARBlendShapeLocationJawRight, EARFaceBlendShape::JawLeft);
+		SET_BLEND_SHAPE(ARBlendShapeLocationMouthLeft, EARFaceBlendShape::MouthRight);
+		SET_BLEND_SHAPE(ARBlendShapeLocationMouthRight, EARFaceBlendShape::MouthLeft);
+		SET_BLEND_SHAPE(ARBlendShapeLocationMouthSmileLeft, EARFaceBlendShape::MouthSmileRight);
+		SET_BLEND_SHAPE(ARBlendShapeLocationMouthSmileRight, EARFaceBlendShape::MouthSmileLeft);
+		SET_BLEND_SHAPE(ARBlendShapeLocationMouthFrownLeft, EARFaceBlendShape::MouthFrownRight);
+		SET_BLEND_SHAPE(ARBlendShapeLocationMouthFrownRight, EARFaceBlendShape::MouthFrownLeft);
+		SET_BLEND_SHAPE(ARBlendShapeLocationMouthDimpleLeft, EARFaceBlendShape::MouthDimpleRight);
+		SET_BLEND_SHAPE(ARBlendShapeLocationMouthDimpleRight, EARFaceBlendShape::MouthDimpleLeft);
+		SET_BLEND_SHAPE(ARBlendShapeLocationMouthStretchLeft, EARFaceBlendShape::MouthStretchRight);
+		SET_BLEND_SHAPE(ARBlendShapeLocationMouthStretchRight, EARFaceBlendShape::MouthStretchLeft);
+		SET_BLEND_SHAPE(ARBlendShapeLocationMouthPressLeft, EARFaceBlendShape::MouthPressRight);
+		SET_BLEND_SHAPE(ARBlendShapeLocationMouthPressRight, EARFaceBlendShape::MouthPressLeft);
+		SET_BLEND_SHAPE(ARBlendShapeLocationMouthLowerDownLeft, EARFaceBlendShape::MouthLowerDownRight);
+		SET_BLEND_SHAPE(ARBlendShapeLocationMouthLowerDownRight, EARFaceBlendShape::MouthLowerDownLeft);
+		SET_BLEND_SHAPE(ARBlendShapeLocationMouthUpperUpLeft, EARFaceBlendShape::MouthUpperUpRight);
+		SET_BLEND_SHAPE(ARBlendShapeLocationMouthUpperUpRight, EARFaceBlendShape::MouthUpperUpLeft);
+		SET_BLEND_SHAPE(ARBlendShapeLocationBrowDownLeft, EARFaceBlendShape::BrowDownRight);
+		SET_BLEND_SHAPE(ARBlendShapeLocationBrowDownRight, EARFaceBlendShape::BrowDownLeft);
+		SET_BLEND_SHAPE(ARBlendShapeLocationBrowOuterUpLeft, EARFaceBlendShape::BrowOuterUpRight);
+		SET_BLEND_SHAPE(ARBlendShapeLocationBrowOuterUpRight, EARFaceBlendShape::BrowOuterUpLeft);
+		SET_BLEND_SHAPE(ARBlendShapeLocationCheekSquintLeft, EARFaceBlendShape::CheekSquintRight);
+		SET_BLEND_SHAPE(ARBlendShapeLocationCheekSquintRight, EARFaceBlendShape::CheekSquintLeft);
+		SET_BLEND_SHAPE(ARBlendShapeLocationNoseSneerLeft, EARFaceBlendShape::NoseSneerRight);
+		SET_BLEND_SHAPE(ARBlendShapeLocationNoseSneerRight, EARFaceBlendShape::NoseSneerLeft);
+	}
+
+	// These are the same mirrored or not
+	SET_BLEND_SHAPE(ARBlendShapeLocationJawOpen, EARFaceBlendShape::JawOpen);
+	SET_BLEND_SHAPE(ARBlendShapeLocationMouthClose, EARFaceBlendShape::MouthClose);
+	SET_BLEND_SHAPE(ARBlendShapeLocationMouthFunnel, EARFaceBlendShape::MouthFunnel);
+	SET_BLEND_SHAPE(ARBlendShapeLocationMouthPucker, EARFaceBlendShape::MouthPucker);
+	SET_BLEND_SHAPE(ARBlendShapeLocationMouthRollLower, EARFaceBlendShape::MouthRollLower);
+	SET_BLEND_SHAPE(ARBlendShapeLocationMouthRollUpper, EARFaceBlendShape::MouthRollUpper);
+	SET_BLEND_SHAPE(ARBlendShapeLocationMouthShrugLower, EARFaceBlendShape::MouthShrugLower);
+	SET_BLEND_SHAPE(ARBlendShapeLocationMouthShrugUpper, EARFaceBlendShape::MouthShrugUpper);
+	SET_BLEND_SHAPE(ARBlendShapeLocationBrowInnerUp, EARFaceBlendShape::BrowInnerUp);
+	SET_BLEND_SHAPE(ARBlendShapeLocationCheekPuff, EARFaceBlendShape::CheekPuff);
+
+#undef SET_BLEND_SHAPE
+
+	return BlendShapeDict;
+}
+
+#endif
+
+UAppleARKitFaceMeshComponent::UAppleARKitFaceMeshComponent(const FObjectInitializer& ObjectInitializer)
+	: Super(ObjectInitializer)
+{
+	PrimaryComponentTick.bCanEverTick = true;
+	// Tick late in the frame to make sure the ARKit has had a chance to update
+	PrimaryComponentTick.TickGroup = TG_PostPhysics;
+	// Init the blend shape map so that all blend shapes are present
+	for (int32 Shape = 0; Shape < (int32)EARFaceBlendShape::MAX; Shape++)
+	{
+		BlendShapes.Add((EARFaceBlendShape)Shape, 0.f);
+	}
+}
+
+void UAppleARKitFaceMeshComponent::InitializeComponent()
+{
+	Super::InitializeComponent();
+
+	if (bAutoBindToLocalFaceMesh)
+	{
+		PrimaryComponentTick.SetTickFunctionEnable(true);
+	}
+}
+
+void UAppleARKitFaceMeshComponent::CreateMesh(const TArray<FVector>& Vertices, const TArray<int32>& Triangles, const TArray<FVector2D>& UV0)
+{
+	// @todo JoeG - add a counter here so we can measure the cost
+	TArray<FVector> Normals;
+	TArray<FLinearColor> VertexColors;
+	TArray<FProcMeshTangent> Tangents;
+
+	CreateMeshSection_LinearColor(0, Vertices, Triangles, Normals, UV0, VertexColors, Tangents, bWantsCollision);
+}
+
+void UAppleARKitFaceMeshComponent::SetBlendShapes(const TMap<EARFaceBlendShape, float>& InBlendShapes)
+{
+	BlendShapes = InBlendShapes;
+	LastUpdateFrameNumber++;
+	LastUpdateTimestamp = FPlatformTime::Seconds();
+	if (LiveLinkSource.IsValid())
+	{
+		LiveLinkSource->PublishBlendShapes(LiveLinkSubjectName, LastUpdateTimestamp, LastUpdateFrameNumber, BlendShapes);
+	}
+}
+
+void UAppleARKitFaceMeshComponent::SetBlendShapeAmount(EARFaceBlendShape BlendShape, float Amount)
+{
+	BlendShapes.Add(BlendShape, Amount);
+}
+
+float UAppleARKitFaceMeshComponent::GetFaceBlendShapeAmount(EARFaceBlendShape BlendShape) const
+{
+	float BlendShapeAmount = 0.f;
+	if (BlendShapes.Contains(BlendShape))
+	{
+		BlendShapeAmount = BlendShapes[BlendShape];
+	}
+	return BlendShapeAmount;
+}
+
+FMatrix UAppleARKitFaceMeshComponent::GetRenderMatrix() const
+{
+	const float Scale = FAppleARKitConversion::ToUE4Scale();
+	
+	FTransform RenderTrans;
+	switch (TransformSetting)
+	{
+		case EARFaceComponentTransformMixing::ComponentOnly:
+		{
+			RenderTrans = GetComponentTransform();
+			break;
+		}
+		case EARFaceComponentTransformMixing::ComponentLocationTrackedRotation:
+		{
+			RenderTrans = GetComponentTransform();
+			FQuat TrackedRot = LocalToWorldTransform.GetRotation();
+			RenderTrans.SetRotation(TrackedRot);
+			break;
+		}
+		case EARFaceComponentTransformMixing::ComponentWithTracked:
+		{
+			RenderTrans = GetComponentTransform() * LocalToWorldTransform;
+			break;
+		}
+		case EARFaceComponentTransformMixing::TrackingOnly:
+		{
+			RenderTrans = LocalToWorldTransform;
+			break;
+		}
+	}
+	RenderTrans.MultiplyScale3D(FVector(Scale));
+	return RenderTrans.ToMatrixWithScale();
+}
+
+class UMaterialInterface* UAppleARKitFaceMeshComponent::GetMaterial(int32 ElementIndex) const
+{
+	if (ElementIndex == 0)
+	{
+		return FaceMaterial;
+	}
+	return nullptr;
+}
+
+void UAppleARKitFaceMeshComponent::SetMaterial(int32 ElementIndex, class UMaterialInterface* Material)
+{
+	if (ElementIndex == 0)
+	{
+		FaceMaterial = Material;
+	}
+}
+
+void UAppleARKitFaceMeshComponent::UpdateMeshFromBlendShapes()
+{
+	LocalToWorldTransform = FTransform::Identity;
+#if SUPPORTS_ARKIT_1_0
+	// @todo JoeG - add a counter here so we can measure the cost
+	NSDictionary<ARBlendShapeLocation,NSNumber *>* BlendShapeDict = ToBlendShapeDictionary(BlendShapes);
+	ARFaceGeometry* FaceGeo = [[ARFaceGeometry alloc] initWithBlendShapes: BlendShapeDict];
+	if (FaceGeo != nullptr)
+	{
+		// Create or update the mesh depending on if we've been created before
+		if (GetNumSections() > 0)
+		{
+			UpdateMesh(ToVertexBuffer(FaceGeo.vertices, FaceGeo.vertexCount));
+		}
+		else
+		{
+			// @todo JoeG - route the uvs in
+			TArray<FVector2D> UVs;
+			CreateMesh(ToVertexBuffer(FaceGeo.vertices, FaceGeo.vertexCount), To32BitIndexBuffer(FaceGeo.triangleIndices, FaceGeo.triangleCount * 3), UVs);
+		}
+	}
+	[FaceGeo release];
+#endif
+}
+
+void UAppleARKitFaceMeshComponent::UpdateMesh(const TArray<FVector>& Vertices)
+{
+	// @todo JoeG - add a counter here so we can measure the cost
+	TArray<FVector> Normals;
+	TArray<FVector2D> UV0;
+	TArray<FLinearColor> VertexColors;
+	TArray<FProcMeshTangent> Tangents;
+
+	UpdateMeshSection_LinearColor(0, Vertices, Normals, UV0, VertexColors, Tangents);
+}
+
+int32 UAppleARKitFaceMeshComponent::GetLastUpdateFrameNumber() const
+{
+	return LastUpdateFrameNumber;
+}
+
+float UAppleARKitFaceMeshComponent::GetLastUpdateTimestamp() const
+{
+	return LastUpdateTimestamp;
+}
+
+UARFaceGeometry* UAppleARKitFaceMeshComponent::FindFaceGeometry()
+{
+	const TArray<UARTrackedGeometry*> Geometries = UARBlueprintLibrary::GetAllGeometries();
+	for (UARTrackedGeometry* Geo : Geometries)
+	{
+		if (UARFaceGeometry* FaceGeo = Cast<UARFaceGeometry>(Geo))
+		{
+			return FaceGeo;
+		}
+	}
+	return nullptr;
+}
+
+void UAppleARKitFaceMeshComponent::SetAutoBind(bool bAutoBind)
+{
+	if (bAutoBindToLocalFaceMesh != bAutoBind)
+	{
+		bAutoBindToLocalFaceMesh = bAutoBind;
+		PrimaryComponentTick.SetTickFunctionEnable(bAutoBind);
+	}
+}
+
+void UAppleARKitFaceMeshComponent::TickComponent(float DeltaTime, ELevelTick, FActorComponentTickFunction*)
+{
+	if (!bAutoBindToLocalFaceMesh)
+	{
+		return;
+	}
+	// Find the tracked face geometry and skip updates if it is not found (can happen if tracking is lost)
+	if (UARFaceGeometry* FaceGeometry = FindFaceGeometry())
+	{
+		LocalToWorldTransform = FaceGeometry->GetLocalToWorldTransform();
+		// +X is the default for Unreal, but you probably want the face point out of the screen (-X)
+		if (bFlipTrackedRotation)
+		{
+			FRotator TrackedRot(LocalToWorldTransform.GetRotation());
+			TrackedRot.Yaw = TrackedRot.Yaw - 180.f;
+			TrackedRot.Pitch = -TrackedRot.Pitch;
+			TrackedRot.Roll = -TrackedRot.Roll;
+			LocalToWorldTransform.SetRotation(FQuat(TrackedRot));
+		}
+		BlendShapes = FaceGeometry->GetBlendShapes();
+		LastUpdateFrameNumber = FaceGeometry->GetLastUpdateFrameNumber();
+		LastUpdateTimestamp = FaceGeometry->GetLastUpdateTimestamp();
+		// Create or update the mesh depending on if we've been created before
+		if (GetNumSections() > 0)
+		{
+			UpdateMesh(FaceGeometry->GetVertexBuffer());
+		}
+		else
+		{
+			CreateMesh(FaceGeometry->GetVertexBuffer(), FaceGeometry->GetIndexBuffer(), FaceGeometry->GetUVs());
+		}
+	}
+}
+
+void UAppleARKitFaceMeshComponent::PublishViaLiveLink(FName SubjectName)
+{
+	LiveLinkSubjectName = SubjectName;
+	if (!LiveLinkSource.IsValid())
+	{
+		LiveLinkSource = FAppleARKitLiveLinkSourceFactory::CreateLiveLinkSource(false);
+	}
+}
+
+FTransform UAppleARKitFaceMeshComponent::GetTransform() const
+{
+	return LocalToWorldTransform;
+}
diff --git a/Engine/Plugins/Runtime/AppleARKitFaceSupport/Source/AppleARKitFaceSupport/Private/AppleARKitFaceMeshConversion.h b/Engine/Plugins/Runtime/AppleARKitFaceSupport/Source/AppleARKitFaceSupport/Private/AppleARKitFaceMeshConversion.h
new file mode 100644
--- /dev/null
+++ b/Engine/Plugins/Runtime/AppleARKitFaceSupport/Source/AppleARKitFaceSupport/Private/AppleARKitFaceMeshConversion.h
@@ -0,0 +1,166 @@
+// Copyright 1998-2018 Epic Games, Inc. All Rights Reserved.
+
+#pragma once
+
+#if SUPPORTS_ARKIT_1_0
+
+FORCEINLINE TArray<int32> To32BitIndexBuffer(const int16_t* Indices, const uint64 IndexCount)
+{
+	check(IndexCount % 3 == 0);
+
+	TArray<int32> IndexBuffer;
+	IndexBuffer.AddUninitialized(IndexCount);
+	for (uint32 Index = 0; Index < IndexCount; Index += 3)
+	{
+		IndexBuffer[Index] = (int32)Indices[Index];
+		// We need to reverse the winding order
+		IndexBuffer[Index + 1] = (int32)Indices[Index + 2];
+		IndexBuffer[Index + 2] = (int32)Indices[Index + 1];
+	}
+	return IndexBuffer;
+}
+
+// @todo JoeG - An option for which way to orient tris (down +X or -X)
+FORCEINLINE TArray<FVector> ToVertexBuffer(const vector_float3* Vertices, const uint64 VertexCount)
+{
+	TArray<FVector> VertexBuffer;
+	VertexBuffer.AddUninitialized(VertexCount);
+	// @todo JoeG - make a fast routine for this
+	for (int32 Index = 0; Index < VertexBuffer.Num(); Index++)
+	{
+		VertexBuffer[Index] = FVector(Vertices[Index].z, Vertices[Index].x, Vertices[Index].y);
+	}
+	return VertexBuffer;
+}
+
+FARBlendShapeMap ToBlendShapeMap(NSDictionary<ARBlendShapeLocation,NSNumber *>* BlendShapes, const FTransform& Transform)
+{
+	FARBlendShapeMap BlendShapeMap;
+	FRotator TrackedRot(Transform.GetRotation());
+	// Map the -180..180 range to -1..1
+	float HeadYaw = (float)TrackedRot.Yaw / 180.f;
+	float HeadPitch = (float)TrackedRot.Pitch / 180.f;
+	float HeadRoll = (float)TrackedRot.Roll / 180.f;
+	BlendShapeMap.Add(EARFaceBlendShape::HeadYaw, HeadYaw);
+	BlendShapeMap.Add(EARFaceBlendShape::HeadPitch, HeadPitch);
+	BlendShapeMap.Add(EARFaceBlendShape::HeadRoll, HeadRoll);
+
+#define SET_BLEND_SHAPE(AppleShape, UE4Shape) \
+	if (BlendShapes[AppleShape]) \
+	{ \
+		BlendShapeMap.Add(UE4Shape, [BlendShapes[AppleShape] floatValue]); \
+	} \
+	else \
+	{ \
+		BlendShapeMap.Add(UE4Shape, 0.f); \
+	}
+
+	// Do we want to capture face performance or look at the face as if in a mirror
+	if (GetDefault<UAppleARKitSettings>()->DefaultFaceTrackingDirection == EARFaceTrackingDirection::FaceRelative)
+	{
+		SET_BLEND_SHAPE(ARBlendShapeLocationEyeBlinkLeft, EARFaceBlendShape::EyeBlinkLeft);
+		SET_BLEND_SHAPE(ARBlendShapeLocationEyeLookDownLeft, EARFaceBlendShape::EyeLookDownLeft);
+		SET_BLEND_SHAPE(ARBlendShapeLocationEyeLookInLeft, EARFaceBlendShape::EyeLookInLeft);
+		SET_BLEND_SHAPE(ARBlendShapeLocationEyeLookOutLeft, EARFaceBlendShape::EyeLookOutLeft);
+		SET_BLEND_SHAPE(ARBlendShapeLocationEyeLookUpLeft, EARFaceBlendShape::EyeLookUpLeft);
+		SET_BLEND_SHAPE(ARBlendShapeLocationEyeSquintLeft, EARFaceBlendShape::EyeSquintLeft);
+		SET_BLEND_SHAPE(ARBlendShapeLocationEyeWideLeft, EARFaceBlendShape::EyeWideLeft);
+		SET_BLEND_SHAPE(ARBlendShapeLocationEyeBlinkRight, EARFaceBlendShape::EyeBlinkRight);
+		SET_BLEND_SHAPE(ARBlendShapeLocationEyeLookDownRight, EARFaceBlendShape::EyeLookDownRight);
+		SET_BLEND_SHAPE(ARBlendShapeLocationEyeLookInRight, EARFaceBlendShape::EyeLookInRight);
+		SET_BLEND_SHAPE(ARBlendShapeLocationEyeLookOutRight, EARFaceBlendShape::EyeLookOutRight);
+		SET_BLEND_SHAPE(ARBlendShapeLocationEyeLookUpRight, EARFaceBlendShape::EyeLookUpRight);
+		SET_BLEND_SHAPE(ARBlendShapeLocationEyeSquintRight, EARFaceBlendShape::EyeSquintRight);
+		SET_BLEND_SHAPE(ARBlendShapeLocationEyeWideRight, EARFaceBlendShape::EyeWideRight);
+		SET_BLEND_SHAPE(ARBlendShapeLocationJawForward, EARFaceBlendShape::JawForward);
+		SET_BLEND_SHAPE(ARBlendShapeLocationJawLeft, EARFaceBlendShape::JawLeft);
+		SET_BLEND_SHAPE(ARBlendShapeLocationJawRight, EARFaceBlendShape::JawRight);
+		SET_BLEND_SHAPE(ARBlendShapeLocationMouthLeft, EARFaceBlendShape::MouthLeft);
+		SET_BLEND_SHAPE(ARBlendShapeLocationMouthRight, EARFaceBlendShape::MouthRight);
+		SET_BLEND_SHAPE(ARBlendShapeLocationMouthSmileLeft, EARFaceBlendShape::MouthSmileLeft);
+		SET_BLEND_SHAPE(ARBlendShapeLocationMouthSmileRight, EARFaceBlendShape::MouthSmileRight);
+		SET_BLEND_SHAPE(ARBlendShapeLocationMouthFrownLeft, EARFaceBlendShape::MouthFrownLeft);
+		SET_BLEND_SHAPE(ARBlendShapeLocationMouthFrownRight, EARFaceBlendShape::MouthFrownRight);
+		SET_BLEND_SHAPE(ARBlendShapeLocationMouthDimpleLeft, EARFaceBlendShape::MouthDimpleLeft);
+		SET_BLEND_SHAPE(ARBlendShapeLocationMouthDimpleRight, EARFaceBlendShape::MouthDimpleRight);
+		SET_BLEND_SHAPE(ARBlendShapeLocationMouthStretchLeft, EARFaceBlendShape::MouthStretchLeft);
+		SET_BLEND_SHAPE(ARBlendShapeLocationMouthStretchRight, EARFaceBlendShape::MouthStretchRight);
+		SET_BLEND_SHAPE(ARBlendShapeLocationMouthPressLeft, EARFaceBlendShape::MouthPressLeft);
+		SET_BLEND_SHAPE(ARBlendShapeLocationMouthPressRight, EARFaceBlendShape::MouthPressRight);
+		SET_BLEND_SHAPE(ARBlendShapeLocationMouthLowerDownLeft, EARFaceBlendShape::MouthLowerDownLeft);
+		SET_BLEND_SHAPE(ARBlendShapeLocationMouthLowerDownRight, EARFaceBlendShape::MouthLowerDownRight);
+		SET_BLEND_SHAPE(ARBlendShapeLocationMouthUpperUpLeft, EARFaceBlendShape::MouthUpperUpLeft);
+		SET_BLEND_SHAPE(ARBlendShapeLocationMouthUpperUpRight, EARFaceBlendShape::MouthUpperUpRight);
+		SET_BLEND_SHAPE(ARBlendShapeLocationBrowDownLeft, EARFaceBlendShape::BrowDownLeft);
+		SET_BLEND_SHAPE(ARBlendShapeLocationBrowDownRight, EARFaceBlendShape::BrowDownRight);
+		SET_BLEND_SHAPE(ARBlendShapeLocationBrowOuterUpLeft, EARFaceBlendShape::BrowOuterUpLeft);
+		SET_BLEND_SHAPE(ARBlendShapeLocationBrowOuterUpRight, EARFaceBlendShape::BrowOuterUpRight);
+		SET_BLEND_SHAPE(ARBlendShapeLocationCheekSquintLeft, EARFaceBlendShape::CheekSquintLeft);
+		SET_BLEND_SHAPE(ARBlendShapeLocationCheekSquintRight, EARFaceBlendShape::CheekSquintRight);
+		SET_BLEND_SHAPE(ARBlendShapeLocationNoseSneerLeft, EARFaceBlendShape::NoseSneerLeft);
+		SET_BLEND_SHAPE(ARBlendShapeLocationNoseSneerRight, EARFaceBlendShape::NoseSneerRight);
+	}
+	else
+	{
+		SET_BLEND_SHAPE(ARBlendShapeLocationEyeBlinkLeft, EARFaceBlendShape::EyeBlinkRight);
+		SET_BLEND_SHAPE(ARBlendShapeLocationEyeLookDownLeft, EARFaceBlendShape::EyeLookDownRight);
+		SET_BLEND_SHAPE(ARBlendShapeLocationEyeLookInLeft, EARFaceBlendShape::EyeLookInRight);
+		SET_BLEND_SHAPE(ARBlendShapeLocationEyeLookOutLeft, EARFaceBlendShape::EyeLookOutRight);
+		SET_BLEND_SHAPE(ARBlendShapeLocationEyeLookUpLeft, EARFaceBlendShape::EyeLookUpRight);
+		SET_BLEND_SHAPE(ARBlendShapeLocationEyeSquintLeft, EARFaceBlendShape::EyeSquintRight);
+		SET_BLEND_SHAPE(ARBlendShapeLocationEyeWideLeft, EARFaceBlendShape::EyeWideRight);
+		SET_BLEND_SHAPE(ARBlendShapeLocationEyeBlinkRight, EARFaceBlendShape::EyeBlinkLeft);
+		SET_BLEND_SHAPE(ARBlendShapeLocationEyeLookDownRight, EARFaceBlendShape::EyeLookDownLeft);
+		SET_BLEND_SHAPE(ARBlendShapeLocationEyeLookInRight, EARFaceBlendShape::EyeLookInLeft);
+		SET_BLEND_SHAPE(ARBlendShapeLocationEyeLookOutRight, EARFaceBlendShape::EyeLookOutLeft);
+		SET_BLEND_SHAPE(ARBlendShapeLocationEyeLookUpRight, EARFaceBlendShape::EyeLookUpLeft);
+		SET_BLEND_SHAPE(ARBlendShapeLocationEyeSquintRight, EARFaceBlendShape::EyeSquintLeft);
+		SET_BLEND_SHAPE(ARBlendShapeLocationEyeWideRight, EARFaceBlendShape::EyeWideLeft);
+		SET_BLEND_SHAPE(ARBlendShapeLocationJawForward, EARFaceBlendShape::JawForward);
+		SET_BLEND_SHAPE(ARBlendShapeLocationJawLeft, EARFaceBlendShape::JawRight);
+		SET_BLEND_SHAPE(ARBlendShapeLocationJawRight, EARFaceBlendShape::JawLeft);
+		SET_BLEND_SHAPE(ARBlendShapeLocationMouthLeft, EARFaceBlendShape::MouthRight);
+		SET_BLEND_SHAPE(ARBlendShapeLocationMouthRight, EARFaceBlendShape::MouthLeft);
+		SET_BLEND_SHAPE(ARBlendShapeLocationMouthSmileLeft, EARFaceBlendShape::MouthSmileRight);
+		SET_BLEND_SHAPE(ARBlendShapeLocationMouthSmileRight, EARFaceBlendShape::MouthSmileLeft);
+		SET_BLEND_SHAPE(ARBlendShapeLocationMouthFrownLeft, EARFaceBlendShape::MouthFrownRight);
+		SET_BLEND_SHAPE(ARBlendShapeLocationMouthFrownRight, EARFaceBlendShape::MouthFrownLeft);
+		SET_BLEND_SHAPE(ARBlendShapeLocationMouthDimpleLeft, EARFaceBlendShape::MouthDimpleRight);
+		SET_BLEND_SHAPE(ARBlendShapeLocationMouthDimpleRight, EARFaceBlendShape::MouthDimpleLeft);
+		SET_BLEND_SHAPE(ARBlendShapeLocationMouthStretchLeft, EARFaceBlendShape::MouthStretchRight);
+		SET_BLEND_SHAPE(ARBlendShapeLocationMouthStretchRight, EARFaceBlendShape::MouthStretchLeft);
+		SET_BLEND_SHAPE(ARBlendShapeLocationMouthPressLeft, EARFaceBlendShape::MouthPressRight);
+		SET_BLEND_SHAPE(ARBlendShapeLocationMouthPressRight, EARFaceBlendShape::MouthPressLeft);
+		SET_BLEND_SHAPE(ARBlendShapeLocationMouthLowerDownLeft, EARFaceBlendShape::MouthLowerDownRight);
+		SET_BLEND_SHAPE(ARBlendShapeLocationMouthLowerDownRight, EARFaceBlendShape::MouthLowerDownLeft);
+		SET_BLEND_SHAPE(ARBlendShapeLocationMouthUpperUpLeft, EARFaceBlendShape::MouthUpperUpRight);
+		SET_BLEND_SHAPE(ARBlendShapeLocationMouthUpperUpRight, EARFaceBlendShape::MouthUpperUpLeft);
+		SET_BLEND_SHAPE(ARBlendShapeLocationBrowDownLeft, EARFaceBlendShape::BrowDownRight);
+		SET_BLEND_SHAPE(ARBlendShapeLocationBrowDownRight, EARFaceBlendShape::BrowDownLeft);
+		SET_BLEND_SHAPE(ARBlendShapeLocationBrowOuterUpLeft, EARFaceBlendShape::BrowOuterUpRight);
+		SET_BLEND_SHAPE(ARBlendShapeLocationBrowOuterUpRight, EARFaceBlendShape::BrowOuterUpLeft);
+		SET_BLEND_SHAPE(ARBlendShapeLocationCheekSquintLeft, EARFaceBlendShape::CheekSquintRight);
+		SET_BLEND_SHAPE(ARBlendShapeLocationCheekSquintRight, EARFaceBlendShape::CheekSquintLeft);
+		SET_BLEND_SHAPE(ARBlendShapeLocationNoseSneerLeft, EARFaceBlendShape::NoseSneerRight);
+		SET_BLEND_SHAPE(ARBlendShapeLocationNoseSneerRight, EARFaceBlendShape::NoseSneerLeft);
+	}
+
+	// These are the same mirrored or not
+	SET_BLEND_SHAPE(ARBlendShapeLocationJawOpen, EARFaceBlendShape::JawOpen);
+	SET_BLEND_SHAPE(ARBlendShapeLocationMouthClose, EARFaceBlendShape::MouthClose);
+	SET_BLEND_SHAPE(ARBlendShapeLocationMouthFunnel, EARFaceBlendShape::MouthFunnel);
+	SET_BLEND_SHAPE(ARBlendShapeLocationMouthPucker, EARFaceBlendShape::MouthPucker);
+	SET_BLEND_SHAPE(ARBlendShapeLocationMouthRollLower, EARFaceBlendShape::MouthRollLower);
+	SET_BLEND_SHAPE(ARBlendShapeLocationMouthRollUpper, EARFaceBlendShape::MouthRollUpper);
+	SET_BLEND_SHAPE(ARBlendShapeLocationMouthShrugLower, EARFaceBlendShape::MouthShrugLower);
+	SET_BLEND_SHAPE(ARBlendShapeLocationMouthShrugUpper, EARFaceBlendShape::MouthShrugUpper);
+	SET_BLEND_SHAPE(ARBlendShapeLocationBrowInnerUp, EARFaceBlendShape::BrowInnerUp);
+	SET_BLEND_SHAPE(ARBlendShapeLocationCheekPuff, EARFaceBlendShape::CheekPuff);
+
+
+#undef SET_BLEND_SHAPE
+
+	return BlendShapeMap;
+}
+
+#endif
diff --git a/Engine/Plugins/Runtime/AppleARKitFaceSupport/Source/AppleARKitFaceSupport/Private/AppleARKitFaceSupportImpl.cpp b/Engine/Plugins/Runtime/AppleARKitFaceSupport/Source/AppleARKitFaceSupport/Private/AppleARKitFaceSupportImpl.cpp
new file mode 100644
--- /dev/null
+++ b/Engine/Plugins/Runtime/AppleARKitFaceSupport/Source/AppleARKitFaceSupport/Private/AppleARKitFaceSupportImpl.cpp
@@ -0,0 +1,181 @@
+// Copyright 1998-2018 Epic Games, Inc. All Rights Reserved.
+
+#include "AppleARKitFaceSupportImpl.h"
+#include "AppleARKitSettings.h"
+#include "AppleARKitFaceMeshConversion.h"
+#include "AppleARKitConfiguration.h"
+#include "Async/TaskGraphInterfaces.h"
+#include "ARSystem.h"
+
+DECLARE_STATS_GROUP(TEXT("AppleARKitFaceSupport"), STATGROUP_APPLEARKITFACE, STATCAT_Advanced);
+
+uint32 LastTrackedFaceGeometryID = 0;
+
+TArray<int32> FAppleARKitFaceAnchorData::FaceIndices;
+
+#if SUPPORTS_ARKIT_1_0
+
+static TSharedPtr<FAppleARKitFaceAnchorData> MakeAnchorData(ARAnchor* Anchor)
+{
+    TSharedPtr<FAppleARKitFaceAnchorData> NewAnchor;
+    if ([Anchor isKindOfClass:[ARFaceAnchor class]])
+    {
+        ARFaceAnchor* FaceAnchor = (ARFaceAnchor*)Anchor;
+        NewAnchor = MakeShared<FAppleARKitFaceAnchorData>(
+			FAppleARKitConversion::ToFGuid(FaceAnchor.identifier),
+			FAppleARKitConversion::ToFTransform(FaceAnchor.transform),
+			ToBlendShapeMap(FaceAnchor.blendShapes, FAppleARKitConversion::ToFTransform(FaceAnchor.transform)),
+			ToVertexBuffer(FaceAnchor.geometry.vertices, FaceAnchor.geometry.vertexCount)
+        );
+        // Only convert from 16bit to 32bit once
+        if (FAppleARKitFaceAnchorData::FaceIndices.Num() == 0)
+        {
+            FAppleARKitFaceAnchorData::FaceIndices = To32BitIndexBuffer(FaceAnchor.geometry.triangleIndices, FaceAnchor.geometry.triangleCount * 3);
+        }
+    }
+
+    return NewAnchor;
+}
+
+#endif
+
+FAppleARKitFaceSupport::FAppleARKitFaceSupport(const TSharedRef<FARSystemBase, ESPMode::ThreadSafe>& InTrackingSystem, IAppleARKitFaceSupportCallback* InCallback)
+    : Callback(InCallback)
+	, ARSystem(InTrackingSystem)
+    , FrameNumber(0)
+    , Timestamp(0.0)
+{
+    check(Callback);
+
+    // Create our LiveLink provider if the project setting is enabled
+    if (GetDefault<UAppleARKitSettings>()->bEnableLiveLinkForFaceTracking)
+    {
+        FaceTrackingLiveLinkSubjectName = GetDefault<UAppleARKitSettings>()->DefaultFaceTrackingLiveLinkSubjectName;
+#if PLATFORM_IOS
+        LiveLinkSource = FAppleARKitLiveLinkSourceFactory::CreateLiveLinkSource(true);
+#else
+        // This should be started already, but just in case
+        FAppleARKitLiveLinkSourceFactory::CreateLiveLinkRemoteListener();
+#endif
+    }
+}
+
+FAppleARKitFaceSupport::~FAppleARKitFaceSupport()
+{
+
+}
+
+#if SUPPORTS_ARKIT_1_0
+
+ARConfiguration* FAppleARKitFaceSupport::ToARConfiguration(UARSessionConfig* SessionConfig, FAppleARKitConfiguration& InConfiguration)
+{
+	ARConfiguration* SessionConfiguration = nullptr;
+	if (SessionConfig->GetSessionType() == EARSessionType::Face)
+	{
+		if (ARFaceTrackingConfiguration.isSupported == FALSE)
+		{
+			return nullptr;
+		}
+		SessionConfiguration = [ARFaceTrackingConfiguration new];
+	}
+
+	// Copy / convert properties
+	SessionConfiguration.lightEstimationEnabled = InConfiguration.bLightEstimationEnabled;
+	SessionConfiguration.providesAudioData = InConfiguration.bProvidesAudioData;
+	SessionConfiguration.worldAlignment = FAppleARKitConversion::ToARWorldAlignment(InConfiguration.Alignment);
+
+	return SessionConfiguration;
+}
+
+void FAppleARKitFaceSupport::ProcessAnchorAdd(NSArray<ARAnchor*>* NewAnchors, const FTransform& InAlignmentTransform, uint32 InFrameNumber, double InTimestamp)
+{
+	DECLARE_CYCLE_STAT(TEXT("FaceAR::ProcessAnchorAdd"), STAT_FaceAR_ProcessAnchorAdd, STATGROUP_APPLEARKITFACE);
+
+    FrameNumber = InFrameNumber;
+    Timestamp = InTimestamp;
+
+	for (ARAnchor* Anchor in NewAnchors)
+	{
+		TSharedPtr<FAppleARKitFaceAnchorData> AnchorData = MakeAnchorData(Anchor);
+		if (AnchorData.IsValid())
+		{
+			auto AnchorTask = FSimpleDelegateGraphTask::FDelegate::CreateSP(this, &FAppleARKitFaceSupport::ProcessAnchorAdd_Internal, AnchorData.ToSharedRef());
+			FSimpleDelegateGraphTask::CreateAndDispatchWhenReady(AnchorTask, GET_STATID(STAT_FaceAR_ProcessAnchorAdd), nullptr, ENamedThreads::GameThread);
+		}
+	}
+}
+
+void FAppleARKitFaceSupport::ProcessAnchorUpdate(NSArray<ARAnchor*>* UpdatedAnchors, const FTransform& InAlignmentTransform, uint32 InFrameNumber, double InTimestamp)
+{
+	DECLARE_CYCLE_STAT(TEXT("FaceAR::ProcessAnchorUpdate"), STAT_FaceAR_ProcessAnchorUpdate, STATGROUP_APPLEARKITFACE);
+
+    FrameNumber = InFrameNumber;
+    Timestamp = InTimestamp;
+
+	for (ARAnchor* Anchor in UpdatedAnchors)
+	{
+		TSharedPtr<FAppleARKitFaceAnchorData> AnchorData = MakeAnchorData(Anchor);
+		if (AnchorData.IsValid())
+		{
+			auto AnchorTask = FSimpleDelegateGraphTask::FDelegate::CreateSP(this, &FAppleARKitFaceSupport::ProcessAnchorUpdate_Internal, AnchorData.ToSharedRef());
+			FSimpleDelegateGraphTask::CreateAndDispatchWhenReady(AnchorTask, GET_STATID(STAT_FaceAR_ProcessAnchorUpdate), nullptr, ENamedThreads::GameThread);
+		}
+	}
+}
+
+void FAppleARKitFaceSupport::ProcessAnchorAdd_Internal(TSharedRef<FAppleARKitFaceAnchorData> AnchorData)
+{
+	if (AnchorData->AnchorType == EAppleAnchorType::FaceAnchor)
+	{
+        // Update LiveLink first, because the other updates use MoveTemp for efficiency
+        if (LiveLinkSource.IsValid())
+        {
+            LiveLinkSource->PublishBlendShapes(FaceTrackingLiveLinkSubjectName, Timestamp, FrameNumber, AnchorData->BlendShapes);
+        }
+
+        FString NewAnchorDebugName;
+        UARFaceGeometry* NewGeo = NewObject<UARFaceGeometry>();
+        NewAnchorDebugName = FString::Printf(TEXT("FACE-%02d"), LastTrackedFaceGeometryID++);
+        NewGeo->SetDebugName(FName(*NewAnchorDebugName));
+
+        NewGeo->UpdateTrackedGeometry(ARSystem, FrameNumber, Timestamp, AnchorData->Transform, AlignmentTransform, AnchorData->BlendShapes, AnchorData->FaceVerts, AnchorData->FaceIndices);
+
+        // Add it to the ar system since that's what other APIs query against
+        Callback->AddTrackedGeometry(AnchorData->AnchorGUID, NewGeo);
+    }
+}
+
+void FAppleARKitFaceSupport::ProcessAnchorUpdate_Internal(TSharedRef<FAppleARKitFaceAnchorData> AnchorData)
+{
+	UARTrackedGeometry* FoundGeometry = Callback->GetTrackedGeometry(AnchorData->AnchorGUID);
+	if (ensure(FoundGeometry != nullptr))
+	{
+		if (UARFaceGeometry* FaceGeo = Cast<UARFaceGeometry>(FoundGeometry))
+		{
+			// Update LiveLink first, because the other updates use MoveTemp for efficiency
+			if (LiveLinkSource.IsValid())
+			{
+				LiveLinkSource->PublishBlendShapes(FaceTrackingLiveLinkSubjectName, Timestamp, FrameNumber, AnchorData->BlendShapes);
+			}
+
+			const TArray<UARPin*>& Pins = ARSystem->GetAllPins();
+			TArray<UARPin*> PinsToUpdate = ARKitUtil::PinsFromGeometry(FoundGeometry, Pins);
+
+			// We figure out the delta transform for the Anchor (aka. TrackedGeometry in ARKit) and apply that
+			// delta to figure out the new ARPin transform.
+			const FTransform Anchor_LocalToTrackingTransform_PreUpdate = FoundGeometry->GetLocalToTrackingTransform_NoAlignment();
+			const FTransform& Anchor_LocalToTrackingTransform_PostUpdate = AnchorData->Transform;
+
+			const FTransform AnchorDeltaTransform = Anchor_LocalToTrackingTransform_PreUpdate.GetRelativeTransform(Anchor_LocalToTrackingTransform_PostUpdate);
+
+			FaceGeo->UpdateTrackedGeometry(ARSystem, FrameNumber, Timestamp, AnchorData->Transform, AlignmentTransform, AnchorData->BlendShapes, AnchorData->FaceVerts, AnchorData->FaceIndices);
+			for (UARPin* Pin : PinsToUpdate)
+			{
+				const FTransform Pin_LocalToTrackingTransform_PostUpdate = Pin->GetLocalToTrackingTransform_NoAlignment() * AnchorDeltaTransform;
+				Pin->OnTransformUpdated(Pin_LocalToTrackingTransform_PostUpdate);
+			}
+		}
+	}
+}
+
+#endif
diff --git a/Engine/Plugins/Runtime/AppleARKitFaceSupport/Source/AppleARKitFaceSupport/Private/AppleARKitFaceSupportImpl.h b/Engine/Plugins/Runtime/AppleARKitFaceSupport/Source/AppleARKitFaceSupport/Private/AppleARKitFaceSupportImpl.h
new file mode 100644
--- /dev/null
+++ b/Engine/Plugins/Runtime/AppleARKitFaceSupport/Source/AppleARKitFaceSupport/Private/AppleARKitFaceSupportImpl.h
@@ -0,0 +1,58 @@
+// Copyright 1998-2018 Epic Games, Inc. All Rights Reserved.
+
+#pragma once
+
+#include "CoreMinimal.h"
+#include "AppleARKitFaceSupport.h"
+#include "AppleARKitLiveLinkSourceFactory.h"
+#include "AppleARKitConversion.h"
+
+
+struct FAppleARKitFaceAnchorData :
+	public FAppleARKitAnchorData
+{
+	FAppleARKitFaceAnchorData(FGuid InAnchorGuid, FTransform InTransform, FARBlendShapeMap InBlendShapes, TArray<FVector> InFaceVerts)
+		: FAppleARKitAnchorData(InAnchorGuid, InTransform)
+		, BlendShapes( MoveTemp(InBlendShapes) )
+		, FaceVerts( MoveTemp(InFaceVerts) )
+	{
+	}
+
+	FARBlendShapeMap BlendShapes;
+	TArray<FVector> FaceVerts;
+	// Note: the index buffer never changes so can be safely read once
+	static TArray<int32> FaceIndices;
+};
+
+class FAppleARKitFaceSupport :
+	public FAppleARKitFaceSupportBase
+{
+public:
+	FAppleARKitFaceSupport(const TSharedRef<FARSystemBase, ESPMode::ThreadSafe>& InTrackingSystem, IAppleARKitFaceSupportCallback* Callback);
+	virtual ~FAppleARKitFaceSupport();
+
+private:
+#if SUPPORTS_ARKIT_1_0
+	// ~IAppleARKitFaceSupport
+	virtual void ProcessAnchorAdd(NSArray<ARAnchor*>* NewAnchors, const FTransform& AlignmentTransform, uint32 FrameNumber, double Timestamp) override;
+	virtual void ProcessAnchorUpdate(NSArray<ARAnchor*>* UpdatedAnchors, const FTransform& AlignmentTransform, uint32 FrameNumber, double Timestamp) override;
+	virtual ARConfiguration* ToARConfiguration(UARSessionConfig* InSessionConfig, FAppleARKitConfiguration& InConfiguration) override;
+	// ~IAppleARKitFaceSupport
+
+	void ProcessAnchorAdd_Internal(TSharedRef<FAppleARKitFaceAnchorData> AnchorData);
+	void ProcessAnchorUpdate_Internal(TSharedRef<FAppleARKitFaceAnchorData> AnchorData);
+#endif
+    
+    IAppleARKitFaceSupportCallback* Callback;
+	const TSharedRef<FARSystemBase, ESPMode::ThreadSafe> ARSystem;
+
+	uint32 FrameNumber;
+	double Timestamp;
+	FTransform AlignmentTransform;
+
+	/** If requested, publishes face ar updates to LiveLink for the animation system to use */
+	TSharedPtr<ILiveLinkSourceARKit> LiveLinkSource;
+	/** Copied from the UARSessionConfig project settings object */
+	FName FaceTrackingLiveLinkSubjectName;
+};
+
diff --git a/Engine/Plugins/Runtime/AppleARKitFaceSupport/Source/AppleARKitFaceSupport/Private/AppleARKitFaceSupportModule.cpp b/Engine/Plugins/Runtime/AppleARKitFaceSupport/Source/AppleARKitFaceSupport/Private/AppleARKitFaceSupportModule.cpp
new file mode 100644
--- /dev/null
+++ b/Engine/Plugins/Runtime/AppleARKitFaceSupport/Source/AppleARKitFaceSupport/Private/AppleARKitFaceSupportModule.cpp
@@ -0,0 +1,76 @@
+// Copyright 1998-2018 Epic Games, Inc. All Rights Reserved.
+
+#include "AppleARKitFaceSupportModule.h"
+#include "AppleARKitSystem.h"
+#include "Features/IModularFeature.h"
+#include "Features/IModularFeatures.h"
+#include "CoreGlobals.h"
+#include "AppleARKitFaceSupportImpl.h"
+
+
+IMPLEMENT_MODULE(FAppleARKitFaceSupportModule, AppleARKitFaceSupport);
+
+DEFINE_LOG_CATEGORY(LogAppleARKitFace);
+
+class FAppleARKitFaceSupportFactory :
+	public IAppleARKitFaceSupportFactory
+{
+public:
+	static void CreateInstance()
+	{
+		Factory = new FAppleARKitFaceSupportFactory();
+		IModularFeatures::Get().RegisterModularFeature(Factory->GetModularFeatureName(), Factory);
+	}
+
+	static void DestroyInstance()
+	{
+		IModularFeatures::Get().UnregisterModularFeature(Factory->GetModularFeatureName(), Factory);
+		delete Factory;
+		Factory = nullptr;
+	}
+
+private:
+	FAppleARKitFaceSupportFactory()
+	{
+	}
+	virtual ~FAppleARKitFaceSupportFactory()
+	{
+	}
+
+	virtual TSharedPtr<FAppleARKitFaceSupportBase, ESPMode::ThreadSafe> CreateFaceSupport(const TSharedRef<FARSystemBase, ESPMode::ThreadSafe>& InTrackingSystem, IAppleARKitFaceSupportCallback* Callback)
+	{
+		if (!FaceARSupport.IsValid())
+		{
+			FaceARSupport = MakeShared<FAppleARKitFaceSupport, ESPMode::ThreadSafe>(InTrackingSystem, Callback);
+		}
+		return FaceARSupport;
+	};
+
+	static FAppleARKitFaceSupportFactory* Factory;
+	TSharedPtr<FAppleARKitFaceSupport, ESPMode::ThreadSafe> FaceARSupport;
+};
+
+FAppleARKitFaceSupportFactory* FAppleARKitFaceSupportFactory::Factory = nullptr;
+
+void FAppleARKitFaceSupportModule::StartupModule()
+{
+	ensureMsgf(FModuleManager::Get().LoadModule("AppleARKit"), TEXT("ARKitFaceSupport depends on the AppleARKit module."));
+
+	FAppleARKitFaceSupportFactory::CreateInstance();
+
+	// LiveLink listener needs to be created here so that the editor can receive remote publishing events
+#if PLATFORM_DESKTOP
+	bool bEnableLiveLinkForFaceTracking = false;
+	GConfig->GetBool(TEXT("/Script/AppleARKit.AppleARKitSettings"), TEXT("bEnableLiveLinkForFaceTracking"), bEnableLiveLinkForFaceTracking, GEngineIni);
+	if (bEnableLiveLinkForFaceTracking)
+	{
+		FAppleARKitLiveLinkSourceFactory::CreateLiveLinkRemoteListener();
+	}
+#endif
+}
+
+void FAppleARKitFaceSupportModule::ShutdownModule()
+{
+	FAppleARKitFaceSupportFactory::DestroyInstance();
+}
+
diff --git a/Engine/Plugins/Runtime/AppleARKitFaceSupport/Source/AppleARKitFaceSupport/Private/AppleARKitFaceSupportModule.h b/Engine/Plugins/Runtime/AppleARKitFaceSupport/Source/AppleARKitFaceSupport/Private/AppleARKitFaceSupportModule.h
new file mode 100644
--- /dev/null
+++ b/Engine/Plugins/Runtime/AppleARKitFaceSupport/Source/AppleARKitFaceSupport/Private/AppleARKitFaceSupportModule.h
@@ -0,0 +1,17 @@
+// Copyright 1998-2018 Epic Games, Inc. All Rights Reserved.
+
+#pragma once
+
+#include "Modules/ModuleManager.h"
+
+class APPLEARKITFACESUPPORT_API FAppleARKitFaceSupportModule :
+	public IModuleInterface
+{
+public:
+	virtual void StartupModule() override;
+
+	virtual void ShutdownModule() override;
+};
+
+DECLARE_LOG_CATEGORY_EXTERN(LogAppleARKitFace, Log, All);
+
diff --git a/Engine/Plugins/Runtime/AppleARKitFaceSupport/Source/AppleARKitFaceSupport/Private/AppleARKitLiveLinkSource.cpp b/Engine/Plugins/Runtime/AppleARKitFaceSupport/Source/AppleARKitFaceSupport/Private/AppleARKitLiveLinkSource.cpp
new file mode 100644
--- /dev/null
+++ b/Engine/Plugins/Runtime/AppleARKitFaceSupport/Source/AppleARKitFaceSupport/Private/AppleARKitLiveLinkSource.cpp
@@ -0,0 +1,395 @@
+// Copyright 1998-2018 Epic Games, Inc. All Rights Reserved.
+
+#include "AppleARKitLiveLinkSource.h"
+#include "UObject/Package.h"
+#include "Misc/ConfigCacheIni.h"
+#include "UObject/UObjectGlobals.h"
+#include "UObject/ObjectMacros.h"
+#include "ARSystem.h"
+#include "ARBlueprintLibrary.h"
+#include "AppleARKitModule.h"
+#include "Features/IModularFeatures.h"
+
+#include "SocketSubsystem.h"
+#include "IPAddress.h"
+#include "Sockets.h"
+
+#include "AppleARKitSettings.h"
+#include "ARTrackable.h"
+
+
+TSharedPtr<ILiveLinkSourceARKit> FAppleARKitLiveLinkSourceFactory::CreateLiveLinkSource(bool bCreateRemotePublisher)
+{
+	IModularFeatures& ModularFeatures = IModularFeatures::Get();
+
+	if (ModularFeatures.IsModularFeatureAvailable(ILiveLinkClient::ModularFeatureName))
+	{
+		ILiveLinkClient* LiveLinkClient = &IModularFeatures::Get().GetModularFeature<ILiveLinkClient>(ILiveLinkClient::ModularFeatureName);
+		TSharedPtr <ILiveLinkSourceARKit> Source = MakeShareable(new FAppleARKitLiveLinkSource(bCreateRemotePublisher));
+		LiveLinkClient->AddSource(Source);
+		return Source;
+	}
+	return TSharedPtr<ILiveLinkSourceARKit>();
+}
+
+void FAppleARKitLiveLinkSourceFactory::CreateLiveLinkRemoteListener()
+{
+	static FAppleARKitLiveLinkRemoteListener* Listener = nullptr;
+	if (Listener == nullptr)
+	{
+		Listener = new FAppleARKitLiveLinkRemoteListener();
+		if (!Listener->InitReceiveSocket())
+		{
+			delete Listener;
+			Listener = nullptr;
+		}
+	}
+}
+
+FAppleARKitLiveLinkSource::FAppleARKitLiveLinkSource(bool bCreateRemotePublisher) :
+	Client(nullptr)
+	, LastFramePublished(0)
+{
+#if PLATFORM_IOS
+	if (bCreateRemotePublisher)
+	{
+		// Only send from iOS to desktop
+		// This will perform the sending of the data to the remote
+		FAppleARKitLiveLinkRemotePublisher* Publisher = new FAppleARKitLiveLinkRemotePublisher();
+		if (Publisher->InitSendSocket())
+		{
+			RemoteLiveLinkPublisher = MakeShareable(Publisher);
+		}
+		else
+		{
+			UE_LOG(LogAppleARKit, Warning, TEXT("Failed to create LiveLink remote publisher, so no data will be sent out"));
+			delete Publisher;
+		}
+	}
+#endif
+}
+
+void FAppleARKitLiveLinkSource::ReceiveClient(ILiveLinkClient* InClient, FGuid InSourceGuid)
+{
+	Client = InClient;
+	SourceGuid = InSourceGuid;
+}
+
+bool FAppleARKitLiveLinkSource::IsSourceStillValid()
+{
+	return Client != nullptr;
+}
+
+bool FAppleARKitLiveLinkSource::RequestSourceShutdown()
+{
+	Client = nullptr;
+	return true;
+}
+
+FText FAppleARKitLiveLinkSource::GetSourceMachineName() const
+{
+	return FText().FromString(FPlatformProcess::ComputerName());
+}
+
+FText FAppleARKitLiveLinkSource::GetSourceStatus() const
+{
+	return NSLOCTEXT( "AppleARKitLiveLink", "AppleARKitLiveLinkStatus", "Active" );
+}
+
+FText FAppleARKitLiveLinkSource::GetSourceType() const
+{
+	return NSLOCTEXT( "AppleARKitLiveLink", "AppleARKitLiveLinkSourceType", "Apple AR Face Tracking" );
+}
+
+static FName ParseEnumName(FName EnumName)
+{
+	const int32 BlendShapeEnumNameLength = 19;
+	FString EnumString = EnumName.ToString();
+	return FName(*EnumString.Right(EnumString.Len() - BlendShapeEnumNameLength));
+}
+
+// Temporary for 4.20
+PRAGMA_DISABLE_DEPRECATION_WARNINGS
+
+void FAppleARKitLiveLinkSource::PublishBlendShapes(FName SubjectName, double Timestamp, uint32 FrameNumber, const FARBlendShapeMap& FaceBlendShapes)
+{
+	check(Client != nullptr);
+
+	if (SubjectName != LastSubjectName)
+	{
+		Client->ClearSubject(LastSubjectName);
+		// We need to publish a skeleton for this subject name even though we doesn't use one
+		Client->PushSubjectSkeleton(SourceGuid, SubjectName, FLiveLinkRefSkeleton());
+	}
+	LastSubjectName = SubjectName;
+
+	const UEnum* EnumPtr = FindObject<UEnum>(ANY_PACKAGE, TEXT("EARFaceBlendShape"), true);
+	if (EnumPtr != nullptr)
+	{
+		static FLiveLinkFrameData LiveLinkFrame;
+
+		LiveLinkFrame.WorldTime = Timestamp;
+
+		TArray<FLiveLinkCurveElement>& BlendShapes = LiveLinkFrame.CurveElements;
+
+		BlendShapes.Reset((int32)EARFaceBlendShape::MAX);
+
+		// Iterate through all of the blend shapes copying them into the LiveLink data type
+		for (int32 Shape = 0; Shape < (int32)EARFaceBlendShape::MAX; Shape++)
+		{
+			if (FaceBlendShapes.Contains((EARFaceBlendShape)Shape))
+			{
+				int32 Index = BlendShapes.AddUninitialized(1);
+				BlendShapes[Index].CurveName = ParseEnumName(EnumPtr->GetNameByValue(Shape));
+				const float CurveValue = FaceBlendShapes.FindChecked((EARFaceBlendShape)Shape);
+				BlendShapes[Index].CurveValue = CurveValue;
+			}
+		}
+
+		// Share the data locally with the LiveLink client
+		Client->PushSubjectData(SourceGuid, SubjectName, LiveLinkFrame);
+		// Send it to the remote editor via the message bus
+		if (RemoteLiveLinkPublisher.IsValid())
+		{
+			RemoteLiveLinkPublisher->PublishBlendShapes(SubjectName, Timestamp, FrameNumber, FaceBlendShapes);
+		}
+	}
+}
+
+PRAGMA_ENABLE_DEPRECATION_WARNINGS
+
+bool FAppleARKitLiveLinkSource::Exec(UWorld* InWorld, const TCHAR* Cmd, FOutputDevice& Ar)
+{
+	if (FParse::Command(&Cmd, TEXT("LiveLinkFaceAR")))
+	{
+		FString RemoteIp;
+		if (FParse::Value(Cmd, TEXT("SendTo="), RemoteIp))
+		{
+			// We need to recreate the LiveLink remote publisher
+			RemoteLiveLinkPublisher = nullptr;
+#if PLATFORM_IOS
+			// Only send from iOS to desktop
+			// This will perform the sending of the data to the remote
+			FAppleARKitLiveLinkRemotePublisher* Publisher = new FAppleARKitLiveLinkRemotePublisher(RemoteIp);
+			if (Publisher->InitSendSocket())
+			{
+				RemoteLiveLinkPublisher = MakeShareable(Publisher);
+			}
+			else
+			{
+				UE_LOG(LogAppleARKit, Warning, TEXT("Failed to create LiveLink remote publisher, so no data will be sent out"));
+				delete Publisher;
+			}
+#endif
+			return true;
+		}
+	}
+	return false;
+}
+
+const uint8 BLEND_SHAPE_PACKET_VER = 1;
+
+const uint32 MAX_BLEND_SHAPE_PACKET_SIZE = sizeof(BLEND_SHAPE_PACKET_VER) + sizeof(double) + sizeof(uint32) + sizeof(uint8) + (sizeof(float) * (uint64)EARFaceBlendShape::MAX) + (sizeof(TCHAR) * 256);
+const uint32 MIN_BLEND_SHAPE_PACKET_SIZE = sizeof(BLEND_SHAPE_PACKET_VER) + sizeof(double) + sizeof(uint32) + sizeof(uint8) + (sizeof(float) * (uint64)EARFaceBlendShape::MAX) + sizeof(TCHAR);
+
+FAppleARKitLiveLinkRemotePublisher::FAppleARKitLiveLinkRemotePublisher(const FString& InRemoteIp) :
+	RemoteIp(InRemoteIp),
+	SendSocket(nullptr),
+	SendBuffer(MAX_BLEND_SHAPE_PACKET_SIZE)
+{
+}
+
+FAppleARKitLiveLinkRemotePublisher::FAppleARKitLiveLinkRemotePublisher() :
+	SendSocket(nullptr),
+	SendBuffer(MAX_BLEND_SHAPE_PACKET_SIZE)
+{
+}
+
+FAppleARKitLiveLinkRemotePublisher::~FAppleARKitLiveLinkRemotePublisher()
+{
+	if (SendSocket != nullptr)
+	{
+		SendSocket->Close();
+		ISocketSubsystem* SocketSub = ISocketSubsystem::Get();
+		SocketSub->DestroySocket(SendSocket);
+	}
+}
+
+bool FAppleARKitLiveLinkRemotePublisher::InitSendSocket()
+{
+	TSharedRef<FInternetAddr> Addr = GetSendAddress();
+	if (Addr->IsValid())
+	{
+		ISocketSubsystem* SocketSubsystem = ISocketSubsystem::Get();
+		// Allocate our socket for sending
+		SendSocket = SocketSubsystem->CreateSocket(NAME_DGram, TEXT("FAppleARKitLiveLinkRemotePublisher socket"), true);
+		SendSocket->SetReuseAddr();
+		SendSocket->SetNonBlocking();
+	}
+	return SendSocket != nullptr;
+}
+
+TSharedRef<FInternetAddr> FAppleARKitLiveLinkRemotePublisher::GetSendAddress()
+{
+	ISocketSubsystem* SocketSub = ISocketSubsystem::Get();
+	TSharedRef<FInternetAddr> SendAddr = SocketSub->CreateInternetAddr();
+	// If we didn't get one passed into the constructor then check the commandline for the IP
+	if (!RemoteIp.Len())
+	{
+		FParse::Value(FCommandLine::Get(), TEXT("LiveLinkRemoteIp="), RemoteIp);
+	}
+	// Don't bother trying to parse the IP if it isn't set
+	if (RemoteIp.Len())
+	{
+		int32 LiveLinkPort = GetDefault<UAppleARKitSettings>()->LiveLinkPublishingPort;
+		SendAddr->SetPort(LiveLinkPort);
+		bool bIsValid = false;
+		SendAddr->SetIp(*RemoteIp, bIsValid);
+		UE_LOG(LogAppleARKit, Log, TEXT("Sending LiveLink face AR data to address (%s)"), *SendAddr->ToString(true));
+	}
+	return SendAddr;
+}
+
+void FAppleARKitLiveLinkRemotePublisher::PublishBlendShapes(FName SubjectName, double Timestamp, uint32 FrameNumber, const FARBlendShapeMap& FaceBlendShapes)
+{
+	if (SendSocket != nullptr)
+	{
+		// Build the packet and send it
+		SendBuffer.Reset();
+		SendBuffer << BLEND_SHAPE_PACKET_VER;
+		SendBuffer << SubjectName;
+		SendBuffer << Timestamp;
+		SendBuffer << FrameNumber;
+		uint8 BlendShapeCount = (uint8)EARFaceBlendShape::MAX;
+		check(FaceBlendShapes.Num() == BlendShapeCount);
+		SendBuffer << BlendShapeCount;
+		// Loop through and send each float for each enum
+		for (uint8 BlendShapeIndex = 0; BlendShapeIndex < BlendShapeCount; BlendShapeIndex++)
+		{
+			SendBuffer << FaceBlendShapes.FindChecked((EARFaceBlendShape)BlendShapeIndex);
+		}
+
+		// Now send the packet
+		uint32 SourceBufferSize = SendBuffer.GetByteCount();
+		check(SourceBufferSize <= MAX_BLEND_SHAPE_PACKET_SIZE && "Max packet size for face blends was exceeded");
+		int32 AmountSent = 0;
+		if (!SendSocket->SendTo(SendBuffer, SourceBufferSize, AmountSent, *GetSendAddress()) ||
+			AmountSent != SourceBufferSize)
+		{
+			ISocketSubsystem* SocketSub = ISocketSubsystem::Get();
+			UE_LOG(LogAppleARKit, Verbose, TEXT("Failed to send face AR packet with error (%s). Packet size (%d), sent (%d)"), SocketSub->GetSocketError(), SourceBufferSize, AmountSent);
+		}
+	}
+}
+
+FAppleARKitLiveLinkRemoteListener::FAppleARKitLiveLinkRemoteListener() :
+	RecvSocket(nullptr)
+{
+	RecvBuffer.AddUninitialized(MAX_BLEND_SHAPE_PACKET_SIZE);
+}
+
+FAppleARKitLiveLinkRemoteListener::~FAppleARKitLiveLinkRemoteListener()
+{
+	if (RecvSocket != nullptr)
+	{
+		RecvSocket->Close();
+		ISocketSubsystem* SocketSub = ISocketSubsystem::Get();
+		SocketSub->DestroySocket(RecvSocket);
+	}
+}
+
+bool FAppleARKitLiveLinkRemoteListener::InitReceiveSocket()
+{
+	ISocketSubsystem* SocketSubsystem = ISocketSubsystem::Get();
+	TSharedRef<FInternetAddr> Addr = SocketSubsystem->GetLocalBindAddr(*GLog);
+	int32 LiveLinkPort = 0;
+	// Have to read this value manually since it happens before UObjects are alive
+	GConfig->GetInt(TEXT("/Script/AppleARKit.AppleARKitSettings"), TEXT("LiveLinkPublishingPort"), LiveLinkPort, GEngineIni);
+	Addr->SetPort(LiveLinkPort);
+
+	RecvSocket = SocketSubsystem->CreateSocket(NAME_DGram, TEXT("FAppleARKitLiveLinkRemoteListener socket"));
+	if (RecvSocket != nullptr)
+	{
+		RecvSocket->SetReuseAddr();
+		RecvSocket->SetNonBlocking();
+		RecvSocket->SetRecvErr();
+		// Bind to our listen port
+		if (!RecvSocket->Bind(*Addr))
+		{
+			SocketSubsystem->DestroySocket(RecvSocket);
+			RecvSocket = nullptr;
+			UE_LOG(LogAppleARKit, Warning, TEXT("Failed to bind to the listen port (%s) for LiveLink face AR receiving with error (%s)"),
+				*Addr->ToString(true), SocketSubsystem->GetSocketError());
+		}
+	}
+	return RecvSocket != nullptr;
+}
+
+void FAppleARKitLiveLinkRemoteListener::InitLiveLinkSource()
+{
+	if (!Source.IsValid())
+	{
+		Source = FAppleARKitLiveLinkSourceFactory::CreateLiveLinkSource(false);
+	}
+}
+
+void FAppleARKitLiveLinkRemoteListener::Tick(float DeltaTime)
+{
+	uint32 BytesPending = 0;
+	while (RecvSocket->HasPendingData(BytesPending))
+	{
+		ISocketSubsystem* SocketSubsystem = ISocketSubsystem::Get();
+		TSharedRef<FInternetAddr> Sender = SocketSubsystem->CreateInternetAddr();
+
+		int32 BytesRead = 0;
+		if (RecvSocket->RecvFrom(RecvBuffer.GetData(), RecvBuffer.Num(), BytesRead, *Sender) &&
+			// Make sure the packet is a complete one and ignore if it is not
+			BytesRead > MIN_BLEND_SHAPE_PACKET_SIZE)
+		{
+			uint8 PacketVer = 0;
+			FName SubjectName;
+			double Timestamp = -1.0;
+			uint32 FrameNumber = 0;
+			uint8 BlendShapeCount = (uint8)EARFaceBlendShape::MAX;
+
+			FNboSerializeFromBuffer FromBuffer(RecvBuffer.GetData(), BytesRead);
+
+			FromBuffer >> PacketVer;
+			if (FromBuffer.HasOverflow() || PacketVer != BLEND_SHAPE_PACKET_VER)
+			{
+				UE_LOG(LogAppleARKit, Verbose, TEXT("Packet overflow reading the packet version for the face AR packet"));
+				return;
+			}
+			FromBuffer >> SubjectName;
+			FromBuffer >> Timestamp;
+			FromBuffer >> FrameNumber;
+			FromBuffer >> BlendShapeCount;
+			if (FromBuffer.HasOverflow() || BlendShapeCount != (uint8)EARFaceBlendShape::MAX)
+			{
+				UE_LOG(LogAppleARKit, Verbose, TEXT("Packet overflow reading the face AR packet's non-array fields"));
+				return;
+			}
+
+			// Loop through and parse each float for each enum
+			for (uint8 BlendShapeIndex = 0; BlendShapeIndex < BlendShapeCount && !FromBuffer.HasOverflow(); BlendShapeIndex++)
+			{
+				float Value = 0.f;
+				FromBuffer >> Value;
+				BlendShapes.Add((EARFaceBlendShape)BlendShapeIndex, Value);
+			}
+			// All of the data was valid, so publish it
+			if (!FromBuffer.HasOverflow())
+			{
+				InitLiveLinkSource();
+				if (Source.IsValid())
+				{
+					Source->PublishBlendShapes(SubjectName, Timestamp, FrameNumber, BlendShapes);
+				}
+			}
+			else
+			{
+				UE_LOG(LogAppleARKit, Verbose, TEXT("Packet overflow reading the face AR packet's array of blend shapes"));
+			}
+		}
+	}
+}
diff --git a/Engine/Plugins/Runtime/AppleARKitFaceSupport/Source/AppleARKitFaceSupport/Private/AppleARKitLiveLinkSource.h b/Engine/Plugins/Runtime/AppleARKitFaceSupport/Source/AppleARKitFaceSupport/Private/AppleARKitLiveLinkSource.h
new file mode 100644
--- /dev/null
+++ b/Engine/Plugins/Runtime/AppleARKitFaceSupport/Source/AppleARKitFaceSupport/Private/AppleARKitLiveLinkSource.h
@@ -0,0 +1,793 @@
+// Copyright 1998-2018 Epic Games, Inc. All Rights Reserved.
+
+#pragma once
+
+#include "CoreMinimal.h"
+#include "AppleARKitLiveLinkSourceFactory.h"
+#include "ILiveLinkClient.h"
+#include "Tickable.h"
+#include "ARTrackable.h"
+
+#include "IPAddress.h"
+#include "Sockets.h"
+//#include "NboSerializer.h"
+
+/**
+ * Hack because our build process doesn't handle plugin dependencies correctly. This will be removed
+ */
+
+/**
+ * Serializes data in network byte order form into a buffer
+ */
+class FNboSerializeToBuffer
+{
+	/** Hidden on purpose */
+	FNboSerializeToBuffer(void);
+
+protected:
+	/**
+	 * Holds the data as it is serialized
+	 */
+	TArray<uint8> Data;
+	/**
+	 * Tracks how many bytes have been written in the packet
+	 */
+	uint32 NumBytes;
+
+	/** Indicates whether writing to the buffer caused an overflow or not */
+	bool bHasOverflowed;
+
+public:
+	/**
+	 * Inits the write tracking
+	 */
+	FNboSerializeToBuffer(uint32 Size) :
+			NumBytes(0),
+			bHasOverflowed(false)
+	{
+		Data.Empty(Size);
+		Data.AddZeroed(Size);
+	}
+
+	/**
+	 * Cast operator to get at the formatted buffer data
+	 */
+	inline operator uint8*(void) const
+	{
+		return (uint8*)Data.GetData();
+	}
+
+	/**
+	 * Cast operator to get at the formatted buffer data
+	 */
+	inline const TArray<uint8>& GetBuffer(void) const
+	{
+		return Data;
+	}
+
+	/**
+	 * Returns the size of the buffer we've built up
+	 */
+	inline uint32 GetByteCount(void) const
+	{
+		return NumBytes;
+	}
+
+	/**
+	 * Returns the number of bytes preallocated in the array
+	 */
+	inline uint32 GetBufferSize(void) const
+	{
+		return (uint32)Data.Num();
+	}
+
+	/**
+	 * Trim any extra space in the buffer that is not used
+	 */
+	inline void TrimBuffer()
+	{
+		if (GetBufferSize() > GetByteCount())
+		{
+			Data.RemoveAt(GetByteCount(), GetBufferSize()-GetByteCount());
+		}
+	}
+
+	/**
+	 * Act as if the buffer is now empty. Useful for buffers that are reused
+	 */
+	inline void Reset()
+	{
+		NumBytes = 0;
+		bHasOverflowed = false;
+	}
+
+	/**
+	 * Adds a char to the buffer
+	 */
+	friend inline FNboSerializeToBuffer& operator<<(FNboSerializeToBuffer& Ar,const char Ch)
+	{
+		if (!Ar.HasOverflow() && Ar.NumBytes + 1 <= Ar.GetBufferSize())
+		{
+			Ar.Data[Ar.NumBytes++] = Ch;
+		}
+		else
+		{
+			Ar.bHasOverflowed = true;
+		}
+
+		return Ar;
+	}
+
+	/**
+	 * Adds a byte to the buffer
+	 */
+	friend inline FNboSerializeToBuffer& operator<<(FNboSerializeToBuffer& Ar,const uint8& B)
+	{
+		if (!Ar.HasOverflow() && Ar.NumBytes + 1 <= Ar.GetBufferSize())
+		{
+			Ar.Data[Ar.NumBytes++] = B;
+		}
+		else
+		{
+			Ar.bHasOverflowed = true;
+		}
+
+		return Ar;
+	}
+
+	/**
+	 * Adds a byte to the buffer
+	 */
+	template<class TEnum>
+	friend inline FNboSerializeToBuffer& operator<<(FNboSerializeToBuffer& Ar,const TEnumAsByte<TEnum>& B)
+	{
+		return Ar << *(uint8*)&B;
+	}
+
+	/**
+	 * Adds an int32 to the buffer
+	 */
+	friend inline FNboSerializeToBuffer& operator<<(FNboSerializeToBuffer& Ar,const int32& I)
+	{
+		return Ar << *(uint32*)&I;
+	}
+
+	/**
+	 * Adds a uint32 to the buffer
+	 */
+	friend inline FNboSerializeToBuffer& operator<<(FNboSerializeToBuffer& Ar,const uint32& D)
+	{
+		if (!Ar.HasOverflow() && Ar.NumBytes + 4 <= Ar.GetBufferSize())
+		{
+			Ar.Data[Ar.NumBytes + 0] = (D >> 24) & 0xFF;
+			Ar.Data[Ar.NumBytes + 1] = (D >> 16) & 0xFF;
+			Ar.Data[Ar.NumBytes + 2] = (D >> 8) & 0xFF;
+			Ar.Data[Ar.NumBytes + 3] = D & 0xFF;
+			Ar.NumBytes += 4;
+		}
+		else
+		{
+			Ar.bHasOverflowed = true;
+		}
+
+		return Ar;
+	}
+
+	/**
+	 * Adds a uint64 to the buffer
+	 */
+	friend inline FNboSerializeToBuffer& operator<<(FNboSerializeToBuffer& Ar,const uint64& Q)
+	{
+		if (!Ar.HasOverflow() && Ar.NumBytes + 8 <= Ar.GetBufferSize())
+		{
+			Ar.Data[Ar.NumBytes + 0] = (Q >> 56) & 0xFF;
+			Ar.Data[Ar.NumBytes + 1] = (Q >> 48) & 0xFF;
+			Ar.Data[Ar.NumBytes + 2] = (Q >> 40) & 0xFF;
+			Ar.Data[Ar.NumBytes + 3] = (Q >> 32) & 0xFF;
+			Ar.Data[Ar.NumBytes + 4] = (Q >> 24) & 0xFF;
+			Ar.Data[Ar.NumBytes + 5] = (Q >> 16) & 0xFF;
+			Ar.Data[Ar.NumBytes + 6] = (Q >> 8) & 0xFF;
+			Ar.Data[Ar.NumBytes + 7] = Q & 0xFF;
+			Ar.NumBytes += 8;
+		}
+		else
+		{
+			Ar.bHasOverflowed = true;
+		}
+
+		return Ar;
+	}
+
+	/**
+	 * Adds a float to the buffer
+	 */
+	friend inline FNboSerializeToBuffer& operator<<(FNboSerializeToBuffer& Ar,const float& F)
+	{
+		uint32 Value = *(uint32*)&F;
+		return Ar << Value;
+	}
+
+	/**
+	 * Adds a double to the buffer
+	 */
+	friend inline FNboSerializeToBuffer& operator<<(FNboSerializeToBuffer& Ar,const double& Dbl)
+	{
+		uint64 Value = *(uint64*)&Dbl;
+		return Ar << Value;
+	}
+
+	/**
+	 * Adds a FString to the buffer
+	 */
+	friend inline FNboSerializeToBuffer& operator<<(FNboSerializeToBuffer& Ar,const FString& String)
+	{
+		// We send strings length prefixed
+		FTCHARToUTF8 Converted(*String);
+		int32 Len = Converted.Length();
+
+		Ar << Len;
+
+		if (!Ar.HasOverflow() && Ar.NumBytes + Len <= Ar.GetBufferSize())
+		{
+			// Handle empty strings
+			if (Len > 0)
+			{
+				ANSICHAR* Ptr = (ANSICHAR*)Converted.Get();
+
+				// memcpy it into the buffer
+				FMemory::Memcpy(&Ar.Data[Ar.NumBytes], Ptr,Len);
+				Ar.NumBytes += Len;
+			}
+		}
+		else
+		{
+			Ar.bHasOverflowed = true;
+		}
+
+		return Ar;
+	}
+
+	/**
+	 * Adds a string to the buffer
+	 */
+	friend inline FNboSerializeToBuffer& operator<<(FNboSerializeToBuffer& Ar,const TCHAR* String)
+	{
+		// We send strings length prefixed (conversion handles null pointers)
+		FTCHARToUTF8 Converted(String);
+		int32 Len = Converted.Length();
+
+		Ar << Len;
+
+		if (!Ar.HasOverflow() && Ar.NumBytes + Len <= Ar.GetBufferSize())
+		{
+			// Handle empty/null strings
+			if (Len > 0)
+			{
+				ANSICHAR* Ptr = (ANSICHAR*)Converted.Get();
+
+				// memcpy it into the buffer
+				FMemory::Memcpy(&Ar.Data[Ar.NumBytes], Ptr, Len);
+				Ar.NumBytes += Len;
+			}
+		}
+		else
+		{
+			Ar.bHasOverflowed = true;
+		}
+
+		return Ar;
+	}
+
+	/**
+	 * Adds a substring to the buffer
+	 */
+	inline FNboSerializeToBuffer& AddString(const ANSICHAR* String,const int32 Length)
+	{
+		// We send strings length prefixed
+		int32 Len = Length;
+		(*this) << Len;
+
+		if (!HasOverflow() && NumBytes + Len <= GetBufferSize())
+		{
+			// Don't process if null
+			if (String)
+			{
+				// memcpy it into the buffer
+				FMemory::Memcpy(&Data[NumBytes],String,Len);
+				NumBytes += Len;
+			}
+		}
+		else
+		{
+			bHasOverflowed = true;
+		}
+
+		return *this;
+	}
+
+	/**
+	 * Adds an FName to the buffer as a string
+	 */
+	friend inline FNboSerializeToBuffer& operator<<(FNboSerializeToBuffer& Ar,const FName& Name)
+	{
+		const FString NameString = Name.ToString();
+		Ar << NameString;
+		return Ar;
+	}
+
+
+	/**
+	 * Adds an ip address to the buffer
+	 */
+	friend inline FNboSerializeToBuffer& operator<<(FNboSerializeToBuffer& Ar,const FInternetAddr& Addr)
+	{
+		uint32 OutIp;
+		Addr.GetIp(OutIp);
+		Ar << OutIp;
+		int32 OutPort;
+		Addr.GetPort(OutPort);
+		Ar << OutPort;
+		return Ar;
+	}
+
+	/**
+	 * Adds a FGuid to the buffer
+	 */
+	friend inline FNboSerializeToBuffer& operator<<(FNboSerializeToBuffer& Ar,const FGuid& Guid)
+	{
+		Ar << Guid.A;
+		Ar << Guid.B;
+		Ar << Guid.C;
+		Ar << Guid.D;
+		return Ar;
+	}
+
+	/**
+	 * Writes a blob of data to the buffer
+	 *
+	 * @param Buffer the source data to append
+	 * @param NumToWrite the size of the blob to write
+	 */
+	inline void WriteBinary(const uint8* Buffer,uint32 NumToWrite)
+	{
+		if (!HasOverflow() && NumBytes + NumToWrite <= GetBufferSize())
+		{
+			FMemory::Memcpy(&Data[NumBytes],Buffer,NumToWrite);
+			NumBytes += NumToWrite;
+		}
+		else
+		{
+			bHasOverflowed = true;
+		}
+	}
+
+	/**
+	 * Gets the buffer at a specific point
+	 */
+	inline uint8* GetRawBuffer(uint32 Offset)
+	{
+		return &Data[Offset];
+	}
+
+	/**
+	 * Skips forward in the buffer by the specified amount
+	 *
+	 * @param Amount the number of bytes to skip ahead
+	 */
+	inline void SkipAheadBy(uint32 Amount)
+	{
+		if (!HasOverflow() && NumBytes + Amount <= GetBufferSize())
+		{
+			NumBytes += Amount;
+		}
+		else
+		{
+			bHasOverflowed = true;
+		}
+	}
+
+	/** Returns whether the buffer had an overflow when writing to it */
+	inline bool HasOverflow(void) const
+	{
+		return bHasOverflowed;
+	}
+};
+
+MSVC_PRAGMA(warning(push))
+// Disable used without initialization warning because the reads are initializing
+MSVC_PRAGMA(warning(disable : 4700))
+
+/**
+ * Class used to read data from a NBO data buffer
+ */
+class FNboSerializeFromBuffer
+{
+protected:
+	/** Pointer to the data this reader is attached to */
+	const uint8* Data;
+	/** The size of the data in bytes */
+	const int32 NumBytes;
+	/** The current location in the byte stream for reading */
+	int32 CurrentOffset;
+	/** Indicates whether reading from the buffer caused an overflow or not */
+	bool bHasOverflowed;
+
+	/** Hidden on purpose */
+	FNboSerializeFromBuffer(void);
+
+public:
+	/**
+	 * Initializes the buffer, size, and zeros the read offset
+	 *
+	 * @param InData the buffer to attach to
+	 * @param Length the size of the buffer we are attaching to
+	 */
+	FNboSerializeFromBuffer(const uint8* InData,int32 Length) :
+			Data(InData),
+			NumBytes(Length),
+			CurrentOffset(0),
+			bHasOverflowed(false)
+	{
+	}
+
+	/**
+	 * Reads a char from the buffer
+	 */
+	friend inline FNboSerializeFromBuffer& operator>>(FNboSerializeFromBuffer& Ar,char& Ch)
+	{
+		if (!Ar.HasOverflow() && Ar.CurrentOffset + 1 <= Ar.NumBytes)
+		{
+			Ch = Ar.Data[Ar.CurrentOffset++];
+		}
+		else
+		{
+			Ar.bHasOverflowed = true;
+		}
+		return Ar;
+	}
+
+	/**
+	 * Reads a byte from the buffer
+	 */
+	friend inline FNboSerializeFromBuffer& operator>>(FNboSerializeFromBuffer& Ar,uint8& B)
+	{
+		if (!Ar.HasOverflow() && Ar.CurrentOffset + 1 <= Ar.NumBytes)
+		{
+			B = Ar.Data[Ar.CurrentOffset++];
+		}
+		else
+		{
+			Ar.bHasOverflowed = true;
+		}
+		return Ar;
+	}
+	/**
+	 * Reads a byte from the buffer
+	 */
+	template<class TEnum>
+	friend inline FNboSerializeFromBuffer& operator>>(FNboSerializeFromBuffer& Ar,TEnumAsByte<TEnum>& B)
+	{
+		return Ar >> *(uint8*)&B;
+	}
+
+	/**
+	 * Reads an int32 from the buffer
+	 */
+	friend inline FNboSerializeFromBuffer& operator>>(FNboSerializeFromBuffer& Ar,int32& I)
+	{
+		return Ar >> *(uint32*)&I;
+	}
+
+	/**
+	 * Reads a uint32 from the buffer
+	 */
+	friend inline FNboSerializeFromBuffer& operator>>(FNboSerializeFromBuffer& Ar,uint32& D)
+	{
+		if (!Ar.HasOverflow() && Ar.CurrentOffset + 4 <= Ar.NumBytes)
+		{
+			uint32 D1 = Ar.Data[Ar.CurrentOffset + 0];
+			uint32 D2 = Ar.Data[Ar.CurrentOffset + 1];
+			uint32 D3 = Ar.Data[Ar.CurrentOffset + 2];
+			uint32 D4 = Ar.Data[Ar.CurrentOffset + 3];
+			D = D1 << 24 | D2 << 16 | D3 << 8 | D4;
+			Ar.CurrentOffset += 4;
+		}
+		else
+		{
+			Ar.bHasOverflowed = true;
+		}
+		return Ar;
+	}
+
+	/**
+	 * Adds a uint64 to the buffer
+	 */
+	friend inline FNboSerializeFromBuffer& operator>>(FNboSerializeFromBuffer& Ar,uint64& Q)
+	{
+		if (!Ar.HasOverflow() && Ar.CurrentOffset + 8 <= Ar.NumBytes)
+		{
+			Q = ((uint64)Ar.Data[Ar.CurrentOffset + 0] << 56) |
+				((uint64)Ar.Data[Ar.CurrentOffset + 1] << 48) |
+				((uint64)Ar.Data[Ar.CurrentOffset + 2] << 40) |
+				((uint64)Ar.Data[Ar.CurrentOffset + 3] << 32) |
+				((uint64)Ar.Data[Ar.CurrentOffset + 4] << 24) |
+				((uint64)Ar.Data[Ar.CurrentOffset + 5] << 16) |
+				((uint64)Ar.Data[Ar.CurrentOffset + 6] << 8) |
+				(uint64)Ar.Data[Ar.CurrentOffset + 7];
+			Ar.CurrentOffset += 8;
+		}
+		else
+		{
+			Ar.bHasOverflowed = true;
+		}
+		return Ar;
+	}
+
+	/**
+	 * Reads a float from the buffer
+	 */
+	friend inline FNboSerializeFromBuffer& operator>>(FNboSerializeFromBuffer& Ar,float& F)
+	{
+		return Ar >> *(uint32*)&F;
+	}
+
+	/**
+	 * Reads a double from the buffer
+	 */
+	friend inline FNboSerializeFromBuffer& operator>>(FNboSerializeFromBuffer& Ar,double& Dbl)
+	{
+		return Ar >> *(uint64*)&Dbl;
+	}
+
+	/**
+	 * Reads a FString from the buffer
+	 */
+	friend inline FNboSerializeFromBuffer& operator>>(FNboSerializeFromBuffer& Ar,FString& String)
+	{
+		// We send strings length prefixed
+		int32 Len = 0;
+		Ar >> Len;
+
+		// Check this way to trust NumBytes and CurrentOffset to be more accurate than the packet Len value
+		const bool bSizeOk = (Len >= 0) && (Len <= (Ar.NumBytes - Ar.CurrentOffset));
+		if (!Ar.HasOverflow() && bSizeOk)
+		{
+			// Handle strings of zero length
+			if (Len > 0)
+			{
+				char* Temp = (char*)FMemory_Alloca(Len + 1);
+				// memcpy it in from the buffer
+				FMemory::Memcpy(Temp, &Ar.Data[Ar.CurrentOffset], Len);
+				Temp[Len] = '\0';
+
+				FUTF8ToTCHAR Converted(Temp);
+				TCHAR* Ptr = (TCHAR*)Converted.Get();
+				String = Ptr;
+				Ar.CurrentOffset += Len;
+			}
+			else
+			{
+				String.Empty();
+			}
+		}
+		else
+		{
+			Ar.bHasOverflowed = true;
+		}
+
+		return Ar;
+	}
+
+	/**
+	 * Reads an FName from the buffer
+	 */
+	friend inline FNboSerializeFromBuffer& operator>>(FNboSerializeFromBuffer& Ar,FName& Name)
+	{
+		FString NameString;
+		Ar >> NameString;
+		Name = FName(*NameString);
+		return Ar;
+	}
+
+	/**
+	 * Reads an ip address from the buffer
+	 */
+	friend inline FNboSerializeFromBuffer& operator>>(FNboSerializeFromBuffer& Ar,FInternetAddr& Addr)
+	{
+		uint32 InIp = 0;
+		Ar >> InIp;
+		Addr.SetIp(InIp);
+		int32 InPort = 0;
+		Ar >> InPort;
+		Addr.SetPort(InPort);
+		return Ar;
+	}
+
+	/**
+	 * Reads a FGuid from the buffer
+	 */
+	friend inline FNboSerializeFromBuffer& operator>>(FNboSerializeFromBuffer& Ar,FGuid& Guid)
+	{
+		Ar >> Guid.A;
+		Ar >> Guid.B;
+		Ar >> Guid.C;
+		Ar >> Guid.D;
+		return Ar;
+	}
+
+	/**
+	 * Reads a blob of data from the buffer
+	 *
+	 * @param OutBuffer the destination buffer
+	 * @param NumToRead the size of the blob to read
+	 */
+	void ReadBinary(uint8* OutBuffer,uint32 NumToRead)
+	{
+		if (!HasOverflow() && CurrentOffset + (int32)NumToRead <= NumBytes)
+		{
+			FMemory::Memcpy(OutBuffer,&Data[CurrentOffset],NumToRead);
+			CurrentOffset += NumToRead;
+		}
+		else
+		{
+			bHasOverflowed = true;
+		}
+	}
+
+	/**
+	 * Seek to the desired position in the buffer
+	 *
+	 * @param Pos the offset from the start of the buffer
+	 */
+	void Seek(int32 Pos)
+	{
+		checkSlow(Pos >= 0);
+
+		if (!HasOverflow() && Pos < NumBytes)
+		{
+			CurrentOffset = Pos;
+		}
+		else
+		{
+			bHasOverflowed = true;
+		}
+	}
+
+	/** @return Current position of the buffer being to be read */
+	inline int32 Tell(void) const
+	{
+		return CurrentOffset;
+	}
+
+	/** Returns whether the buffer had an overflow when reading from it */
+	inline bool HasOverflow(void) const
+	{
+		return bHasOverflowed;
+	}
+
+	/** @return Number of bytes remaining to read from the current offset to the end of the buffer */
+	inline int32 AvailableToRead(void) const
+	{
+		return FMath::Max<int32>(0,NumBytes - CurrentOffset);
+	}
+
+	/**
+	 * Returns the number of total bytes in the buffer
+	 */
+	inline int32 GetBufferSize(void) const
+	{
+		return NumBytes;
+	}
+};
+
+MSVC_PRAGMA(warning(pop))
+
+/**
+ * End hack
+ */
+
+class UARFaceGeometry;
+
+/** Listens for LiveLink publication events and publishes them to the local LiveLinkClient */
+class FAppleARKitLiveLinkRemoteListener :
+	public FTickableGameObject
+{
+public:
+	FAppleARKitLiveLinkRemoteListener();
+	virtual ~FAppleARKitLiveLinkRemoteListener();
+
+	// FTickableGameObject interface
+	virtual void Tick(float DeltaTime) override;
+	virtual bool IsTickable() const override { return RecvSocket != nullptr; }
+	virtual TStatId GetStatId() const override
+	{
+		RETURN_QUICK_DECLARE_CYCLE_STAT(FAppleARKitLiveLinkRemoteListener, STATGROUP_Tickables);
+	}
+	virtual bool IsTickableInEditor() const override { return true; }
+	virtual bool IsTickableWhenPaused() const override { return true; }
+	// End FTickableGameObject
+
+	bool InitReceiveSocket();
+
+private:
+	void InitLiveLinkSource();
+
+	/** Socket that is read from to publish remote data to this local LiveLinkClient */
+	FSocket* RecvSocket;
+	/** Buffer to reuse for receiving data */
+	TArray<uint8> RecvBuffer;
+	/** The reused blend shape map to avoid allocs/frees */
+	FARBlendShapeMap BlendShapes;
+
+	/** The source used to publish events from the remote socket */
+	TSharedPtr<ILiveLinkSourceARKit> Source;
+};
+
+/** Forwards face blend shapes to another machine for LiveLink publication */
+class FAppleARKitLiveLinkRemotePublisher :
+	public IARKitBlendShapePublisher
+{
+public:
+	FAppleARKitLiveLinkRemotePublisher();
+	FAppleARKitLiveLinkRemotePublisher(const FString& InRemoteIp);
+	virtual ~FAppleARKitLiveLinkRemotePublisher();
+
+	// IARKitBlendShapePublisher interface
+	virtual void PublishBlendShapes(FName SubjectName, double Timestamp, uint32 FrameNumber, const FARBlendShapeMap& FaceBlendShapes) override;
+	// End IARKitBlendShapePublisher
+
+	bool InitSendSocket();
+
+private:
+	TSharedRef<FInternetAddr> GetSendAddress();
+
+	/** The IP address to send the blend shapes to */
+	FString RemoteIp;
+	/** Socket to send the data via to the remote listener */
+	FSocket* SendSocket;
+	/** The ram buffer to serialize the packet to */
+	FNboSerializeToBuffer SendBuffer;
+};
+
+/** Publishes face blend shapes to LiveLink for use locally */
+class FAppleARKitLiveLinkSource :
+	public ILiveLinkSourceARKit,
+	public FSelfRegisteringExec
+{
+public:
+	FAppleARKitLiveLinkSource(bool bCreateRemotePublisher);
+	virtual ~FAppleARKitLiveLinkSource() {}
+
+private:
+	// ILiveLinkSource interface
+	virtual void ReceiveClient(ILiveLinkClient* InClient, FGuid InSourceGuid) override;
+	virtual bool IsSourceStillValid() override;
+	virtual bool RequestSourceShutdown() override;
+	virtual FText GetSourceMachineName() const override;
+	virtual FText GetSourceStatus() const override;
+	virtual FText GetSourceType() const override;
+	// End ILiveLinkSource
+
+	// IARKitBlendShapePublisher interface
+	virtual void PublishBlendShapes(FName SubjectName, double Timestamp, uint32 FrameNumber, const FARBlendShapeMap& FaceBlendShapes) override;
+	// End IARKitBlendShapePublisher
+
+	//~ FSelfRegisteringExec
+	virtual bool Exec(UWorld* InWorld, const TCHAR* Cmd, FOutputDevice& Ar) override;
+	//~ FSelfRegisteringExec
+
+	/** The local client to push data updates to */
+	ILiveLinkClient* Client;
+
+	/** Our identifier in LiveLink */
+	FGuid SourceGuid;
+
+	/** Provider to push data to remote listeners */
+	TSharedPtr<IARKitBlendShapePublisher> RemoteLiveLinkPublisher;
+
+	/** The last time we sent the data. Used to not send redundant data */
+	uint32 LastFramePublished;
+
+	/** Used to track names changes */
+	FName LastSubjectName;
+};
diff --git a/Engine/Plugins/Runtime/AppleARKitFaceSupport/Source/AppleARKitFaceSupport/Public/AppleARKitFaceMeshComponent.h b/Engine/Plugins/Runtime/AppleARKitFaceSupport/Source/AppleARKitFaceSupport/Public/AppleARKitFaceMeshComponent.h
new file mode 100644
--- /dev/null
+++ b/Engine/Plugins/Runtime/AppleARKitFaceSupport/Source/AppleARKitFaceSupport/Public/AppleARKitFaceMeshComponent.h
@@ -0,0 +1,155 @@
+// Copyright 1998-2018 Epic Games, Inc. All Rights Reserved.
+
+#pragma once
+
+#include "AppleARKitAvailability.h"
+#include "ARSystem.h"
+#include "ProceduralMeshComponent.h"
+#include "AppleARKitLiveLinkSourceFactory.h"
+#include "ARTrackable.h"
+#include "AppleARKitFaceMeshComponent.generated.h"
+
+
+UENUM(BlueprintType, Category="AR AugmentedReality", meta=(Experimental))
+enum class EARFaceComponentTransformMixing : uint8
+{
+	/** Uses the component's transform exclusively. Only setting for non-tracked meshes */
+	ComponentOnly,
+	/** Use the component's location and apply the rotation from the tracked mesh */
+	ComponentLocationTrackedRotation,
+	/** Concatenate the component and the tracked face transforms */
+	ComponentWithTracked,
+	/** Use only the tracked face transform */
+	TrackingOnly
+};
+
+/**
+ * This component is updated by the ARSystem with face data on devices that have support for it
+ */
+UCLASS(hidecategories = (Object, LOD, "Components|ProceduralMesh"), meta = (BlueprintSpawnableComponent), ClassGroup = "AR")
+class APPLEARKIT_API UAppleARKitFaceMeshComponent :
+	public UProceduralMeshComponent
+{
+	GENERATED_UCLASS_BODY()
+
+public:
+	/**
+	 *	Create the initial face mesh from raw mesh data
+	 *
+	 *	@param	Vertices			Vertex buffer of all vertex positions to use for this mesh section.
+	 *	@param	Triangles			Index buffer indicating which vertices make up each triangle. Length must be a multiple of 3.
+	 *	@param	UV0					Optional array of texture co-ordinates for each vertex. If supplied, must be same length as Vertices array.
+	 */
+	UFUNCTION(BlueprintCallable, Category = "Components|ARFaceMesh", meta = (DisplayName = "Create Face Mesh", AutoCreateRefTerm = "UV0"))
+	void CreateMesh(const TArray<FVector>& Vertices, const TArray<int32>& Triangles, const TArray<FVector2D>& UV0);
+
+	/**
+	 *	Set all of the blend shapes for this instance from a set of blend shapes
+	 *
+	 *	@param	BlendShapes			The set of blend shapes to conform the face mesh to
+	 */
+	UFUNCTION(BlueprintCallable, Category = "Components|ARFaceMesh", meta = (DisplayName = "Create Face Mesh from Blend Shapes"))
+	void SetBlendShapes(const TMap<EARFaceBlendShape, float>& InBlendShapes);
+
+	/**
+	 *	Sets the amount for a given blend shape
+	 *
+	 *	@param	BlendShape			The blend shape to modify
+	 *	@param	Amount				The amount of the blend shape to apply (0..1)
+	 */
+	UFUNCTION(BlueprintCallable, Category = "Components|ARFaceMesh", meta = (DisplayName = "Set the value of a Blend Shape"))
+	void SetBlendShapeAmount(EARFaceBlendShape BlendShape, float Amount);
+
+	/**
+	 * Returns the value of the specified blend shape
+	 *
+	 * @param BlendShape			The blend shape to query
+	 *
+	 * @return the current amount the blend shape should be applied
+	 */
+	UFUNCTION(BlueprintPure, Category = "Components|ARFaceMesh")
+	float GetFaceBlendShapeAmount(EARFaceBlendShape BlendShape) const;
+
+	/**
+	 *	Create/replace the face mesh from the current set of blend shapes if the device supports it
+	 */
+	UFUNCTION(BlueprintCallable, Category = "Components|ARFaceMesh", meta = (DisplayName = "Update Face Mesh from Blend Shapes"))
+	void UpdateMeshFromBlendShapes();
+
+	/**
+	 *	Updates the face mesh vertices. The topology and UVs do not change post creation so only vertices are updated
+	 *
+	 *	@param	Vertices			Vertex buffer of all vertex positions to use for this mesh section.
+	 */
+	UFUNCTION(BlueprintCallable, Category = "Components|ARFaceMesh", meta = (DisplayName = "Update Mesh Section FColor"))
+	void UpdateMesh(const TArray<FVector>& Vertices);
+
+	/**
+	 * If auto bind is true, then this component will update itself from the local face tracking data each tick. If auto bind is off, ticking is disabled
+	 *
+	 * @param	bAutoBind			true to enable, false to disable
+	 */
+	UFUNCTION(BlueprintCallable, Category = "Components|ARFaceMesh", meta = (DisplayName = "Modify auto bind to local face tracking"))
+	void SetAutoBind(bool bAutoBind);
+
+	/** Returns the frame number that was last used to update this component */
+	UFUNCTION(BlueprintPure, Category = "Components|ARFaceMesh")
+	int32 GetLastUpdateFrameNumber() const;
+
+	/** Returns the frame timestamp that was last used to update this component */
+	UFUNCTION(BlueprintPure, Category = "Components|ARFaceMesh")
+	float GetLastUpdateTimestamp() const;
+
+	/** Starts LiveLink publishing of this face component's data so that it can be used by the animation system */
+	UFUNCTION(BlueprintCallable, Category = "Components|ARFaceMesh")
+	void PublishViaLiveLink(FName SubjectName);
+
+	/** Get the transform that the AR camera has detected */
+	UFUNCTION(BlueprintPure, Category = "Components|ARFaceMesh")
+	FTransform GetTransform() const;
+
+	/**	Indicates whether collision should be created for this face mesh. This adds significant cost, so only use if you need to trace against the face mesh. */
+	UPROPERTY(EditAnywhere, BlueprintReadOnly, Category = "Components|ARFaceMesh")
+	bool bWantsCollision;
+
+	/**	If true, the mesh data will come from the local ARKit face mesh data. The face mesh will update every tick and will handle loss of face tracking */
+	UPROPERTY(EditAnywhere, BlueprintReadOnly, Category = "Components|ARFaceMesh")
+	bool bAutoBindToLocalFaceMesh;
+
+	/**	Determines how the transform from tracking data and the component's transform are mixed together */
+	UPROPERTY(EditAnywhere, BlueprintReadOnly, Category = "Components|ARFaceMesh")
+	EARFaceComponentTransformMixing TransformSetting;
+
+	/**	If true, the face mesh will be rotated to face out of the screen (-X) rather than into the screen (+X) and corresponding axises to match */
+	UPROPERTY(EditAnywhere, BlueprintReadOnly, Category = "Components|ARFaceMesh")
+	bool bFlipTrackedRotation;
+
+	/** Used when rendering the face mesh (mostly debug reasons) */
+	UPROPERTY(EditAnywhere, BlueprintReadOnly, Category = "Components|ARFaceMesh")
+	class UMaterialInterface* FaceMaterial;
+
+	/** Used to identify this component's face ar data uniquely as part of the LiveLink animation pipeline */
+	UPROPERTY(EditAnywhere, BlueprintReadOnly, Category = "Components|ARFaceMesh")
+	FName LiveLinkSubjectName;
+
+private:
+	virtual void TickComponent(float DeltaTime, enum ELevelTick TickType, FActorComponentTickFunction* ThisTickFunction) override;
+	virtual void InitializeComponent() override;
+	virtual FMatrix GetRenderMatrix() const override;
+	virtual class UMaterialInterface* GetMaterial(int32 ElementIndex) const override;
+	virtual void SetMaterial(int32 ElementIndex, class UMaterialInterface* Material) override;
+
+	UARFaceGeometry* FindFaceGeometry();
+
+	/** The current set of blend shapes for this component instance */
+	FARBlendShapeMap BlendShapes;
+	/** Transform of the face mesh */
+	FTransform LocalToWorldTransform;
+
+	/** The frame number this component was last updated on */
+	uint32 LastUpdateFrameNumber;
+	/** The time reported by the AR system that this object was last updated */
+	double LastUpdateTimestamp;
+	/** If this mesh is being published via LiveLink, the source to update with blendshapes */
+	TSharedPtr<ILiveLinkSourceARKit> LiveLinkSource;
+};
diff --git a/Engine/Plugins/Runtime/AppleARKitFaceSupport/Source/AppleARKitFaceSupport/Public/AppleARKitLiveLinkSourceFactory.h b/Engine/Plugins/Runtime/AppleARKitFaceSupport/Source/AppleARKitFaceSupport/Public/AppleARKitLiveLinkSourceFactory.h
new file mode 100644
--- /dev/null
+++ b/Engine/Plugins/Runtime/AppleARKitFaceSupport/Source/AppleARKitFaceSupport/Public/AppleARKitLiveLinkSourceFactory.h
@@ -0,0 +1,33 @@
+// Copyright 1998-2018 Epic Games, Inc. All Rights Reserved.
+
+#pragma once
+
+#include "ARSystem.h"
+#include "ARTrackable.h"
+#include "ILiveLinkSource.h"
+
+
+/** Interface that publishes face ar blend shape information */
+class APPLEARKIT_API IARKitBlendShapePublisher
+{
+public:
+	virtual void PublishBlendShapes(FName SubjectName, double Timestamp, uint32 FrameNumber, const FARBlendShapeMap& BlendShapes) = 0;
+};
+
+/** Interface that publishes face ar blend shape information via LiveLink */
+class APPLEARKIT_API ILiveLinkSourceARKit :
+	public IARKitBlendShapePublisher,
+	public ILiveLinkSource
+{
+};
+
+/** Factory that creates and registers the sources with the LiveLink client */
+class APPLEARKIT_API FAppleARKitLiveLinkSourceFactory
+{
+public:
+	/** Creates a face mesh source that will autobind to the tracked face mesh */
+	static TSharedPtr<ILiveLinkSourceARKit> CreateLiveLinkSource(bool bCreateRemotePublisher);
+
+	/** Creates the listener that will receive remote events and publish them locally */
+	static void CreateLiveLinkRemoteListener();
+};
