Moved the dependency on IOSRuntimeSettings to be only when building for iOS

diff --git a/Engine/Plugins/Runtime/AppleARKit/Source/AppleARKit/AppleARKit.Build.cs b/Engine/Plugins/Runtime/AppleARKit/Source/AppleARKit/AppleARKit.Build.cs
--- a/Engine/Plugins/Runtime/AppleARKit/Source/AppleARKit/AppleARKit.Build.cs
+++ b/Engine/Plugins/Runtime/AppleARKit/Source/AppleARKit/AppleARKit.Build.cs
@@ -35,7 +35,6 @@ public class AppleARKit : ModuleRules
                 "RenderCore",
                 "ShaderCore",
                 "HeadMountedDisplay",
-                "IOSRuntimeSettings",
                 "AugmentedReality",
                 "AppleImageUtils"
 				// ... add private dependencies that you statically link with here ...
@@ -52,6 +51,8 @@ public class AppleARKit : ModuleRules
 
 		if (Target.Platform == UnrealTargetPlatform.IOS)
 		{
+    		PrivateDependencyModuleNames.Add("IOSRuntimeSettings");
+
 			PublicFrameworks.Add( "ARKit" );
 		}
 	}
diff --git a/Engine/Plugins/Runtime/AppleARKitFaceSupport/Source/AppleARKitFaceSupport/AppleARKitFaceSupport.Build.cs b/Engine/Plugins/Runtime/AppleARKitFaceSupport/Source/AppleARKitFaceSupport/AppleARKitFaceSupport.Build.cs
--- a/Engine/Plugins/Runtime/AppleARKitFaceSupport/Source/AppleARKitFaceSupport/AppleARKitFaceSupport.Build.cs
+++ b/Engine/Plugins/Runtime/AppleARKitFaceSupport/Source/AppleARKitFaceSupport/AppleARKitFaceSupport.Build.cs
@@ -28,7 +28,6 @@ public class AppleARKitFaceSupport : ModuleRules
 			new string[]
 			{
 				"CoreUObject",
-                "IOSRuntimeSettings",
                 "HeadMountedDisplay",
                 "AugmentedReality",
                 "ProceduralMeshComponent",
