Fixed a threading bug that appeared when migrating the face ar code into its own plugin

diff --git a/Engine/Plugins/Runtime/AppleARKit/Source/AppleARKit/Private/AppleARKitSystem.cpp b/Engine/Plugins/Runtime/AppleARKit/Source/AppleARKit/Private/AppleARKitSystem.cpp
--- a/Engine/Plugins/Runtime/AppleARKit/Source/AppleARKit/Private/AppleARKitSystem.cpp
+++ b/Engine/Plugins/Runtime/AppleARKit/Source/AppleARKit/Private/AppleARKitSystem.cpp
@@ -1104,6 +1104,12 @@ void FAppleARKitSystem::SessionDidAddAnchors_Internal( TSharedRef<FAppleARKitAnc
 		UpdateFrame();
 	}
 
+	// If this object is valid, we are running a face session and we need to publish LiveLink data on the game thread
+	if (FaceARSupport != nullptr && AnchorData->AnchorType == EAppleAnchorType::FaceAnchor)
+	{
+		FaceARSupport->PublishLiveLinkData(AnchorData, GameThreadTimestamp, GameThreadFrameNumber);
+	}
+
 	FString NewAnchorDebugName;
 	UARTrackedGeometry* NewGeometry = nullptr;
 	switch (AnchorData->AnchorType)
@@ -1162,6 +1168,12 @@ void FAppleARKitSystem::SessionDidUpdateAnchors_Internal( TSharedRef<FAppleARKit
 		UpdateFrame();
 	}
 
+	// If this object is valid, we are running a face session and we need to publish LiveLink data on the game thread
+	if (FaceARSupport != nullptr && AnchorData->AnchorType == EAppleAnchorType::FaceAnchor)
+	{
+		FaceARSupport->PublishLiveLinkData(AnchorData, GameThreadTimestamp, GameThreadFrameNumber);
+	}
+
 	UARTrackedGeometry** GeometrySearchResult = TrackedGeometries.Find(AnchorData->AnchorGUID);
 	if (ensure(GeometrySearchResult != nullptr))
 	{
diff --git a/Engine/Plugins/Runtime/AppleARKit/Source/AppleARKit/Public/AppleARKitFaceSupport.h b/Engine/Plugins/Runtime/AppleARKit/Source/AppleARKit/Public/AppleARKitFaceSupport.h
--- a/Engine/Plugins/Runtime/AppleARKit/Source/AppleARKit/Public/AppleARKitFaceSupport.h
+++ b/Engine/Plugins/Runtime/AppleARKit/Source/AppleARKit/Public/AppleARKitFaceSupport.h
@@ -31,6 +31,18 @@ public:
 	 */
 	virtual TArray<TSharedPtr<FAppleARKitAnchorData>> MakeAnchorData(NSArray<ARAnchor*>* NewAnchors, double Timestamp, uint32 FrameNumber) = 0;
 
+	/**
+	 * Publishes any face AR data that needs to be sent to LiveLink. Done as a separate step because MakeAnchorData is called
+	 * on an arbitrary thread and we can't access UObjects there safely
+	 *
+	 * @param AnchorList the list of anchors to publish to LiveLink
+	 * @param Timestamp the timestamp of this update
+	 * @param FrameNumber the frame number for this update
+	 *
+	 * @return the set of face anchors to dispatch
+	 */
+	virtual void PublishLiveLinkData(TSharedPtr<FAppleARKitAnchorData> Anchor, double Timestamp, uint32 FrameNumber) = 0;
+
 	/**
 	 * Creates a face ar specific configuration object if that is requested without exposing the main code to the face APIs
 	 *
@@ -39,9 +51,6 @@ public:
 	 */
 	virtual ARConfiguration* ToARConfiguration(UARSessionConfig* SessionConfig, FAppleARKitConfiguration& InConfiguration) = 0;
 #endif
-
-	IAppleARKitFaceSupport() { }
-	virtual ~IAppleARKitFaceSupport() { }
 };
 
 class APPLEARKIT_API IAppleARKitFaceSupportFactory :
diff --git a/Engine/Plugins/Runtime/AppleARKitFaceSupport/Source/AppleARKitFaceSupport/Private/AppleARKitFaceSupportImpl.cpp b/Engine/Plugins/Runtime/AppleARKitFaceSupport/Source/AppleARKitFaceSupport/Private/AppleARKitFaceSupportImpl.cpp
--- a/Engine/Plugins/Runtime/AppleARKitFaceSupport/Source/AppleARKitFaceSupport/Private/AppleARKitFaceSupportImpl.cpp
+++ b/Engine/Plugins/Runtime/AppleARKitFaceSupport/Source/AppleARKitFaceSupport/Private/AppleARKitFaceSupportImpl.cpp
@@ -87,11 +87,6 @@ TArray<TSharedPtr<FAppleARKitAnchorData>> FAppleARKitFaceSupport::MakeAnchorData
 		TSharedPtr<FAppleARKitAnchorData> AnchorData = ::MakeAnchorData(Anchor);
 		if (AnchorData.IsValid())
 		{
-			// Update LiveLink first, because the other updates use MoveTemp for efficiency
-			if (LiveLinkSource.IsValid())
-			{
-				LiveLinkSource->PublishBlendShapes(FaceTrackingLiveLinkSubjectName, Timestamp, FrameNumber, AnchorData->BlendShapes);
-			}
 			AnchorList.Add(AnchorData);
 		}
 	}
@@ -99,4 +94,12 @@ TArray<TSharedPtr<FAppleARKitAnchorData>> FAppleARKitFaceSupport::MakeAnchorData
 	return AnchorList;
 }
 
+void FAppleARKitFaceSupport::PublishLiveLinkData(TSharedPtr<FAppleARKitAnchorData> Anchor, double Timestamp, uint32 FrameNumber)
+{
+	if (LiveLinkSource.IsValid())
+	{
+        LiveLinkSource->PublishBlendShapes(FaceTrackingLiveLinkSubjectName, Timestamp, FrameNumber, Anchor->BlendShapes);
+	}
+}
+
 #endif
diff --git a/Engine/Plugins/Runtime/AppleARKitFaceSupport/Source/AppleARKitFaceSupport/Private/AppleARKitFaceSupportImpl.h b/Engine/Plugins/Runtime/AppleARKitFaceSupport/Source/AppleARKitFaceSupport/Private/AppleARKitFaceSupportImpl.h
--- a/Engine/Plugins/Runtime/AppleARKitFaceSupport/Source/AppleARKitFaceSupport/Private/AppleARKitFaceSupportImpl.h
+++ b/Engine/Plugins/Runtime/AppleARKitFaceSupport/Source/AppleARKitFaceSupport/Private/AppleARKitFaceSupportImpl.h
@@ -20,9 +20,10 @@ private:
 	// ~IAppleARKitFaceSupport
 	virtual TArray<TSharedPtr<FAppleARKitAnchorData>> MakeAnchorData(NSArray<ARAnchor*>* NewAnchors, double Timestamp, uint32 FrameNumber) override;
 	virtual ARConfiguration* ToARConfiguration(UARSessionConfig* InSessionConfig, FAppleARKitConfiguration& InConfiguration) override;
+	virtual void PublishLiveLinkData(TSharedPtr<FAppleARKitAnchorData> Anchor, double Timestamp, uint32 FrameNumber) override;
 	// ~IAppleARKitFaceSupport
 #endif
-    
+
 	/** If requested, publishes face ar updates to LiveLink for the animation system to use */
 	TSharedPtr<ILiveLinkSourceARKit> LiveLinkSource;
 	/** Copied from the UARSessionConfig project settings object */
diff --git a/Engine/Plugins/Runtime/AppleARKitFaceSupport/Source/AppleARKitFaceSupport/Private/AppleARKitLiveLinkSource.cpp b/Engine/Plugins/Runtime/AppleARKitFaceSupport/Source/AppleARKitFaceSupport/Private/AppleARKitLiveLinkSource.cpp
--- a/Engine/Plugins/Runtime/AppleARKitFaceSupport/Source/AppleARKitFaceSupport/Private/AppleARKitLiveLinkSource.cpp
+++ b/Engine/Plugins/Runtime/AppleARKitFaceSupport/Source/AppleARKitFaceSupport/Private/AppleARKitLiveLinkSource.cpp
@@ -116,6 +116,8 @@ PRAGMA_DISABLE_DEPRECATION_WARNINGS
 void FAppleARKitLiveLinkSource::PublishBlendShapes(FName SubjectName, double Timestamp, uint32 FrameNumber, const FARBlendShapeMap& FaceBlendShapes)
 {
 	check(Client != nullptr);
+	// This code touches UObjects so needs to be run only in the game thread
+	check(IsInGameThread());
 
 	if (SubjectName != LastSubjectName)
 	{
