Made sure that GetPose() was skipped if tracking was disabled

diff --git a/Engine/Plugins/Runtime/AppleARKit/Source/AppleARKit/Private/AppleARKitSystem.cpp b/Engine/Plugins/Runtime/AppleARKit/Source/AppleARKit/Private/AppleARKitSystem.cpp
--- a/Engine/Plugins/Runtime/AppleARKit/Source/AppleARKit/Private/AppleARKitSystem.cpp
+++ b/Engine/Plugins/Runtime/AppleARKit/Source/AppleARKit/Private/AppleARKitSystem.cpp
@@ -210,7 +210,7 @@ FName FAppleARKitSystem::GetSystemName() const
 
 bool FAppleARKitSystem::GetCurrentPose(int32 DeviceId, FQuat& OutOrientation, FVector& OutPosition)
 {
-	if (DeviceId == IXRTrackingSystem::HMDDeviceId && GameThreadFrame.IsValid())
+	if (DeviceId == IXRTrackingSystem::HMDDeviceId && GameThreadFrame.IsValid() && IsHeadTrackingAllowed())
 	{
 		// Do not have to lock here, because we are on the game
 		// thread and GameThreadFrame is only written to from the game thread.
