Added a setting to specify the type of AR world alignment transform type a session should use
Deleted some old ARKit configuration code that wasn't really used
Fixed the setting of light estimation using the wrong value to determine the setting on ARKit

diff --git a/Engine/Plugins/Runtime/AppleARKit/Source/AppleARKit/Private/AppleARKitConfiguration.cpp b/Engine/Plugins/Runtime/AppleARKit/Source/AppleARKit/Private/AppleARKitConfiguration.cpp
deleted file mode 100644
diff --git a/Engine/Plugins/Runtime/AppleARKit/Source/AppleARKit/Private/AppleARKitConfiguration.h b/Engine/Plugins/Runtime/AppleARKit/Source/AppleARKit/Private/AppleARKitConfiguration.h
deleted file mode 100644
diff --git a/Engine/Plugins/Runtime/AppleARKit/Source/AppleARKit/Private/AppleARKitSystem.cpp b/Engine/Plugins/Runtime/AppleARKit/Source/AppleARKit/Private/AppleARKitSystem.cpp
--- a/Engine/Plugins/Runtime/AppleARKit/Source/AppleARKit/Private/AppleARKitSystem.cpp
+++ b/Engine/Plugins/Runtime/AppleARKit/Source/AppleARKit/Private/AppleARKitSystem.cpp
@@ -10,7 +10,7 @@
 #include "AppleARKitFrame.h"
 #include "AppleARKitAnchor.h"
 #include "AppleARKitPlaneAnchor.h"
-#include "AppleARKitConfiguration.h"
+#include "AppleARKitConversion.h"
 #include "GeneralProjectSettings.h"
 #include "ARSessionConfig.h"
 #include "AppleARKitSettings.h"
@@ -828,8 +828,6 @@ bool FAppleARKitSystem::Run(UARSessionConfig* SessionConfig)
 		return true;
 	}
 
-	// @todo arkit FAppleARKitSystem::GetWorldToMetersScale needs a real scale somehow
-	FAppleARKitConfiguration Config;
 #if SUPPORTS_ARKIT_1_0
 	if (FAppleARKitAvailability::SupportsARKit10())
 	{
@@ -839,11 +837,11 @@ bool FAppleARKitSystem::Run(UARSessionConfig* SessionConfig)
 		CheckForFaceARSupport(SessionConfig);
 		if (FaceARSupport == nullptr)
 		{
-			Configuration = ToARConfiguration(SessionConfig, Config, CandidateImages, ConvertedCandidateImages);
+			Configuration = FAppleARKitConversion::ToARConfiguration(SessionConfig, CandidateImages, ConvertedCandidateImages);
 		}
 		else
 		{
-			Configuration = FaceARSupport->ToARConfiguration(SessionConfig, Config);
+			Configuration = FaceARSupport->ToARConfiguration(SessionConfig);
 		}
 
 		// Not all session types are supported by all devices
diff --git a/Engine/Plugins/Runtime/AppleARKit/Source/AppleARKit/Private/AppleARKitSystem.h b/Engine/Plugins/Runtime/AppleARKit/Source/AppleARKit/Private/AppleARKitSystem.h
--- a/Engine/Plugins/Runtime/AppleARKit/Source/AppleARKit/Private/AppleARKitSystem.h
+++ b/Engine/Plugins/Runtime/AppleARKit/Source/AppleARKit/Private/AppleARKitSystem.h
@@ -3,7 +3,6 @@
 #pragma once
 
 #include "XRTrackingSystemBase.h"
-#include "AppleARKitConfiguration.h"
 #include "ARSystem.h"
 #include "AppleARKitHitTestResult.h"
 #include "AppleARKitTextures.h"
diff --git a/Engine/Plugins/Runtime/AppleARKit/Source/AppleARKit/Public/AppleARKitConversion.h b/Engine/Plugins/Runtime/AppleARKit/Source/AppleARKit/Public/AppleARKitConversion.h
--- a/Engine/Plugins/Runtime/AppleARKit/Source/AppleARKit/Public/AppleARKitConversion.h
+++ b/Engine/Plugins/Runtime/AppleARKit/Source/AppleARKit/Public/AppleARKitConversion.h
@@ -6,25 +6,7 @@
 #include "Math/Transform.h"
 #include "ARPin.h"
 
-//@todo -- Move this to UARSessionConfiguration
-/**
- * Enum constants for indicating the world alignment.
- */
-enum class EAppleARKitWorldAlignment : uint8
-{
-	/** Aligns the world with gravity that is defined by vector (0, -1, 0). */
-	Gravity,
-
-	/**
-	 * Aligns the world with gravity that is defined by the vector (0, -1, 0)
-	 * and heading (w.r.t. True North) that is given by the vector (0, 0, -1).
-	 */
-	GravityAndHeading,
-
-	/** Aligns the world with the camera's orientation. */
-	Camera
-};
-ENUM_CLASS_FLAGS(EAppleARKitWorldAlignment)
+#include "IAppleImageUtilsPlugin.h"
 
 struct FAppleARKitConversion
 {
@@ -93,20 +75,128 @@ struct FAppleARKitConversion
 		return ToFGuid( UUID );
 	}
 
-	static FORCEINLINE ARWorldAlignment ToARWorldAlignment( const EAppleARKitWorldAlignment& InWorldAlignment )
+	static FORCEINLINE ARWorldAlignment ToARWorldAlignment( const EARWorldAlignment& InWorldAlignment )
 	{
 		switch ( InWorldAlignment )
 		{
-			case EAppleARKitWorldAlignment::Gravity:
+			case EARWorldAlignment::Gravity:
 				return ARWorldAlignmentGravity;
 
-			case EAppleARKitWorldAlignment::GravityAndHeading:
+			case EARWorldAlignment::GravityAndHeading:
 				return ARWorldAlignmentGravityAndHeading;
 
-			case EAppleARKitWorldAlignment::Camera:
+			case EARWorldAlignment::Camera:
 				return ARWorldAlignmentCamera;
 		};
 	};
+
+	#pragma clang diagnostic push
+	#pragma clang diagnostic ignored "-Wpartial-availability"
+
+	#if SUPPORTS_ARKIT_1_5
+
+	static void InitImageDetection(UARSessionConfig* SessionConfig, ARWorldTrackingConfiguration* WorldConfig, TMap< FString, UARCandidateImage* >& CandidateImages, TMap< FString, CGImageRef >& ConvertedCandidateImages)
+	{
+		const TArray<UARCandidateImage*>& ConfigCandidateImages = SessionConfig->GetCandidateImageList();
+		if (!ConfigCandidateImages.Num())
+		{
+			return;
+		}
+
+		NSMutableSet* ConvertedImageSet = [[NSMutableSet new] autorelease];
+		for (UARCandidateImage* Candidate : ConfigCandidateImages)
+		{
+			if (Candidate != nullptr && Candidate->GetCandidateTexture() != nullptr)
+			{
+				// Store off so the session object can quickly match the anchor to our representation
+				// This stores it even if we weren't able to convert to apple's type for GC reasons
+				CandidateImages.Add(Candidate->GetFriendlyName(), Candidate);
+				// Convert our texture to an Apple compatible image type
+				CGImageRef ConvertedImage = nullptr;
+				// Avoid doing the expensive conversion work if it's in the cache already
+				CGImageRef* FoundImage = ConvertedCandidateImages.Find(Candidate->GetFriendlyName());
+				if (FoundImage != nullptr)
+				{
+					ConvertedImage = *FoundImage;
+				}
+				else
+				{
+					ConvertedImage = IAppleImageUtilsPlugin::Get().UTexture2DToCGImage(Candidate->GetCandidateTexture());
+					// If it didn't convert this time, it never will, so always store it off
+					ConvertedCandidateImages.Add(Candidate->GetFriendlyName(), ConvertedImage);
+				}
+				if (ConvertedImage != nullptr)
+				{
+					float ImageWidth = (float)Candidate->GetPhysicalWidth() / 100.f;
+					CGImagePropertyOrientation Orientation = Candidate->GetOrientation() == EARCandidateImageOrientation::Landscape ? kCGImagePropertyOrientationRight : kCGImagePropertyOrientationUp;
+
+					ARReferenceImage* ReferenceImage = [[[ARReferenceImage alloc] initWithCGImage: ConvertedImage orientation: Orientation physicalWidth: ImageWidth] autorelease];
+					ReferenceImage.name = Candidate->GetFriendlyName().GetNSString();
+					[ConvertedImageSet addObject: ReferenceImage];
+				}
+			}
+		}
+		WorldConfig.detectionImages = ConvertedImageSet;
+	}
+
+	#endif
+
+	static ARConfiguration* ToARConfiguration( UARSessionConfig* SessionConfig, TMap< FString, UARCandidateImage* >& CandidateImages, TMap< FString, CGImageRef >& ConvertedCandidateImages )
+	{
+		EARSessionType SessionType = SessionConfig->GetSessionType();
+		ARConfiguration* SessionConfiguration = nullptr;
+		switch (SessionType)
+		{
+			case EARSessionType::Orientation:
+			{
+				if (AROrientationTrackingConfiguration.isSupported == FALSE)
+				{
+					return nullptr;
+				}
+				SessionConfiguration = [AROrientationTrackingConfiguration new];
+				break;
+			}
+			case EARSessionType::World:
+			{
+				if (ARWorldTrackingConfiguration.isSupported == FALSE)
+				{
+					return nullptr;
+				}
+				ARWorldTrackingConfiguration* WorldTrackingConfiguration = [ARWorldTrackingConfiguration new];
+				WorldTrackingConfiguration.planeDetection = ARPlaneDetectionNone;
+				if ( EnumHasAnyFlags(EARPlaneDetectionMode::HorizontalPlaneDetection, SessionConfig->GetPlaneDetectionMode()))
+				{
+					WorldTrackingConfiguration.planeDetection |= ARPlaneDetectionHorizontal;
+				}
+#if SUPPORTS_ARKIT_1_5
+				if (FAppleARKitAvailability::SupportsARKit15())
+				{
+					if (EnumHasAnyFlags(EARPlaneDetectionMode::VerticalPlaneDetection, SessionConfig->GetPlaneDetectionMode()) )
+					{
+						WorldTrackingConfiguration.planeDetection |= ARPlaneDetectionVertical;
+					}
+					WorldTrackingConfiguration.autoFocusEnabled = SessionConfig->ShouldEnableAutoFocus();
+					// Add any images that wish to be detected
+					FAppleARKitConversion::InitImageDetection(SessionConfig, WorldTrackingConfiguration, CandidateImages, ConvertedCandidateImages);
+				}
+#endif
+				SessionConfiguration = WorldTrackingConfiguration;
+				break;
+			}
+			default:
+				return nullptr;
+		}
+		check(SessionConfiguration != nullptr);
+
+		// Copy / convert properties
+		SessionConfiguration.lightEstimationEnabled = SessionConfig->GetLightEstimationMode() != EARLightEstimationMode::None;
+		SessionConfiguration.providesAudioData = NO;
+		SessionConfiguration.worldAlignment = FAppleARKitConversion::ToARWorldAlignment(SessionConfig->GetWorldAlignment());
+
+		return SessionConfiguration;
+	}
+
+	#pragma clang diagnostic pop
 #endif
 };
 
diff --git a/Engine/Plugins/Runtime/AppleARKit/Source/AppleARKit/Public/AppleARKitFaceSupport.h b/Engine/Plugins/Runtime/AppleARKit/Source/AppleARKit/Public/AppleARKitFaceSupport.h
--- a/Engine/Plugins/Runtime/AppleARKit/Source/AppleARKit/Public/AppleARKitFaceSupport.h
+++ b/Engine/Plugins/Runtime/AppleARKit/Source/AppleARKit/Public/AppleARKitFaceSupport.h
@@ -12,7 +12,6 @@
 #endif
 
 class UARSessionConfig;
-class FAppleARKitConfiguration;
 struct FAppleARKitAnchorData;
 
 
@@ -47,9 +46,8 @@ public:
 	 * Creates a face ar specific configuration object if that is requested without exposing the main code to the face APIs
 	 *
 	 * @param SessionConfig the UE4 configuration object that needs processing
-	 * @param InConfiguration the legacy configuration object that needs to go away
 	 */
-	virtual ARConfiguration* ToARConfiguration(UARSessionConfig* SessionConfig, FAppleARKitConfiguration& InConfiguration) = 0;
+	virtual ARConfiguration* ToARConfiguration(UARSessionConfig* SessionConfig) = 0;
 #endif
 };
 
diff --git a/Engine/Plugins/Runtime/AppleARKitFaceSupport/Source/AppleARKitFaceSupport/Private/AppleARKitFaceSupportImpl.cpp b/Engine/Plugins/Runtime/AppleARKitFaceSupport/Source/AppleARKitFaceSupport/Private/AppleARKitFaceSupportImpl.cpp
--- a/Engine/Plugins/Runtime/AppleARKitFaceSupport/Source/AppleARKitFaceSupport/Private/AppleARKitFaceSupportImpl.cpp
+++ b/Engine/Plugins/Runtime/AppleARKitFaceSupport/Source/AppleARKitFaceSupport/Private/AppleARKitFaceSupportImpl.cpp
@@ -3,7 +3,7 @@
 #include "AppleARKitFaceSupportImpl.h"
 #include "AppleARKitSettings.h"
 #include "AppleARKitFaceMeshConversion.h"
-#include "AppleARKitConfiguration.h"
+#include "AppleARKitConversion.h"
 #include "Async/TaskGraphInterfaces.h"
 #include "ARSystem.h"
 #include "Misc/ConfigCacheIni.h"
@@ -53,12 +53,13 @@ FAppleARKitFaceSupport::FAppleARKitFaceSupport()
 
 FAppleARKitFaceSupport::~FAppleARKitFaceSupport()
 {
-
+	// Should only be called durirng shutdown
+	check(GIsRequestingExit);
 }
 
 #if SUPPORTS_ARKIT_1_0
 
-ARConfiguration* FAppleARKitFaceSupport::ToARConfiguration(UARSessionConfig* SessionConfig, FAppleARKitConfiguration& InConfiguration)
+ARConfiguration* FAppleARKitFaceSupport::ToARConfiguration(UARSessionConfig* SessionConfig)
 {
 	ARConfiguration* SessionConfiguration = nullptr;
 	if (SessionConfig->GetSessionType() == EARSessionType::Face)
@@ -71,9 +72,9 @@ ARConfiguration* FAppleARKitFaceSupport::ToARConfiguration(UARSessionConfig* Ses
 	}
 
 	// Copy / convert properties
-	SessionConfiguration.lightEstimationEnabled = InConfiguration.bLightEstimationEnabled;
-	SessionConfiguration.providesAudioData = InConfiguration.bProvidesAudioData;
-	SessionConfiguration.worldAlignment = FAppleARKitConversion::ToARWorldAlignment(InConfiguration.Alignment);
+	SessionConfiguration.lightEstimationEnabled = SessionConfig->GetLightEstimationMode() != EARLightEstimationMode::None;
+	SessionConfiguration.providesAudioData = NO;
+	SessionConfiguration.worldAlignment = FAppleARKitConversion::ToARWorldAlignment(SessionConfig->GetWorldAlignment());
 
 	return SessionConfiguration;
 }
diff --git a/Engine/Plugins/Runtime/AppleARKitFaceSupport/Source/AppleARKitFaceSupport/Private/AppleARKitFaceSupportImpl.h b/Engine/Plugins/Runtime/AppleARKitFaceSupport/Source/AppleARKitFaceSupport/Private/AppleARKitFaceSupportImpl.h
--- a/Engine/Plugins/Runtime/AppleARKitFaceSupport/Source/AppleARKitFaceSupport/Private/AppleARKitFaceSupportImpl.h
+++ b/Engine/Plugins/Runtime/AppleARKitFaceSupport/Source/AppleARKitFaceSupport/Private/AppleARKitFaceSupportImpl.h
@@ -19,7 +19,7 @@ private:
 #if SUPPORTS_ARKIT_1_0
 	// ~IAppleARKitFaceSupport
 	virtual TArray<TSharedPtr<FAppleARKitAnchorData>> MakeAnchorData(NSArray<ARAnchor*>* NewAnchors, double Timestamp, uint32 FrameNumber) override;
-	virtual ARConfiguration* ToARConfiguration(UARSessionConfig* InSessionConfig, FAppleARKitConfiguration& InConfiguration) override;
+	virtual ARConfiguration* ToARConfiguration(UARSessionConfig* InSessionConfig) override;
 	virtual void PublishLiveLinkData(TSharedPtr<FAppleARKitAnchorData> Anchor, double Timestamp, uint32 FrameNumber) override;
 	// ~IAppleARKitFaceSupport
 #endif
diff --git a/Engine/Source/Runtime/AugmentedReality/Private/ARSessionConfig.cpp b/Engine/Source/Runtime/AugmentedReality/Private/ARSessionConfig.cpp
--- a/Engine/Source/Runtime/AugmentedReality/Private/ARSessionConfig.cpp
+++ b/Engine/Source/Runtime/AugmentedReality/Private/ARSessionConfig.cpp
@@ -4,7 +4,8 @@
 #include "UObject/VRObjectVersion.h"
 
 UARSessionConfig::UARSessionConfig()
-: SessionType(EARSessionType::World)
+: WorldAlignment(EARWorldAlignment::Gravity)
+, SessionType(EARSessionType::World)
 , PlaneDetectionMode_DEPRECATED(EARPlaneDetectionMode::HorizontalPlaneDetection)
 , bHorizontalPlaneDetection(true)
 , bVerticalPlaneDetection(true)
@@ -15,6 +16,12 @@ UARSessionConfig::UARSessionConfig()
 , bEnableAutomaticCameraTracking(true)
 {
 }
+
+EARWorldAlignment UARSessionConfig::GetWorldAlignment() const
+{
+	return WorldAlignment;
+}
+
 EARSessionType UARSessionConfig::GetSessionType() const
 {
 	return SessionType;
diff --git a/Engine/Source/Runtime/AugmentedReality/Public/ARSessionConfig.h b/Engine/Source/Runtime/AugmentedReality/Public/ARSessionConfig.h
--- a/Engine/Source/Runtime/AugmentedReality/Public/ARSessionConfig.h
+++ b/Engine/Source/Runtime/AugmentedReality/Public/ARSessionConfig.h
@@ -9,6 +9,22 @@
 
 #include "ARSessionConfig.generated.h"
 
+UENUM(BlueprintType, Category="AR AugmentedReality", meta=(Experimental))
+enum class EARWorldAlignment : uint8
+{
+	/** Aligns the world with gravity that is defined by vector (0, -1, 0) */
+	Gravity,
+
+	/**
+	 * Aligns the world with gravity that is defined by the vector (0, -1, 0)
+	 * and heading (w.r.t. True North) that is given by the vector (0, 0, -1)
+	 */
+	GravityAndHeading,
+
+	/** Aligns the world with the camera's orientation, which is best for Face AR */
+	Camera
+};
+
 UENUM(BlueprintType, Category = "AR AugmentedReality", meta = (Experimental))
 enum class EARSessionType : uint8
 {
@@ -71,6 +87,9 @@ public:
 	UARSessionConfig();
 	
 public:
+	/** @see EARWorldAlignment */
+	EARWorldAlignment GetWorldAlignment() const;
+
 	/** @see SessionType */
 	EARSessionType GetSessionType() const;
 
@@ -101,7 +120,9 @@ private:
 	//~ UObject interface
 
 protected:
-	
+	/** @see EARWorldAlignment */
+	EARWorldAlignment WorldAlignment;
+
 	/** @see EARSessionType */
 	UPROPERTY(EditAnywhere, Category = "AR Settings")
 	EARSessionType SessionType;
