First pass integration of ARKit 2.0 (very wip, thar be dragons everywhere)

diff --git a/Engine/Plugins/Experimental/AppleImageUtils/Source/AppleImageUtils/Public/AppleImageUtilsTypes.h b/Engine/Plugins/Experimental/AppleImageUtils/Source/AppleImageUtils/Public/AppleImageUtilsTypes.h
--- a/Engine/Plugins/Experimental/AppleImageUtils/Source/AppleImageUtils/Public/AppleImageUtilsTypes.h
+++ b/Engine/Plugins/Experimental/AppleImageUtils/Source/AppleImageUtils/Public/AppleImageUtilsTypes.h
@@ -46,7 +46,9 @@ public:
 	virtual CVPixelBufferRef GetPixelBuffer() const { return nullptr; }
 	/** @return the IOSurface held by the implementing object */
 	virtual IOSurfaceRef GetSurface() const { return nullptr; }
-	/** @return the CVMetalTexture held by the implementing object */
-	virtual CVMetalTextureRef GetMetalTexture() const { return nullptr; }
+
+//@joeg -- Added for environment capture support
+	/** @return the MTLTextureid<MTL held by the implementing object */
+	virtual id<MTLTexture> GetMetalTexture() const { return nullptr; }
 #endif
 };
diff --git a/Engine/Plugins/Runtime/AppleARKit/Source/AppleARKit/Private/AppleARKitFrame.cpp b/Engine/Plugins/Runtime/AppleARKit/Source/AppleARKit/Private/AppleARKitFrame.cpp
--- a/Engine/Plugins/Runtime/AppleARKit/Source/AppleARKit/Private/AppleARKitFrame.cpp
+++ b/Engine/Plugins/Runtime/AppleARKit/Source/AppleARKit/Private/AppleARKitFrame.cpp
@@ -15,11 +15,30 @@ FAppleARKitFrame::FAppleARKitFrame()
 {
 };
 
+#if SUPPORTS_ARKIT_2_0
+EARWorldMappingState ToEARWorldMappingState(ARWorldMappingStatus MapStatus)
+{
+	switch (MapStatus)
+	{
+		case ARWorldMappingStatusLimited:
+			return EARWorldMappingState::StillMappingNotRelocalizable;
+
+		case ARWorldMappingStatusExtending:
+			return EARWorldMappingState::StillMappingRelocalizable;
+
+		case ARWorldMappingStatusMapped:
+			return EARWorldMappingState::Mapped;
+	}
+	return EARWorldMappingState::NotAvailable;
+}
+#endif
+
 #if SUPPORTS_ARKIT_1_0
 
 FAppleARKitFrame::FAppleARKitFrame( ARFrame* InARFrame, CVMetalTextureCacheRef MetalTextureCache )
   : Camera( InARFrame.camera )
   , LightEstimate( InARFrame.lightEstimate )
+	, WorldMappingState(EARWorldMappingState::NotAvailable)
 {
 	// Sanity check
 	check( InARFrame );
@@ -68,13 +87,21 @@ FAppleARKitFrame::FAppleARKitFrame( ARFrame* InARFrame, CVMetalTextureCacheRef M
 		check( CFGetRetainCount(CapturedCbCrImage) == 1);
 	}
 	// @todo JoeG -- finsih the depth capture
-	if (InARFrame.capturedDepthData)
+//@joeg -- Disabled due to crashing when accessing
+	if (0 && InARFrame.capturedDepthData)
 	{
 		CameraDepth = InARFrame.capturedDepthData;
 		CFRetain(CameraDepth);
 	}
 
 	NativeFrame = (void*)CFRetain(InARFrame);
+
+#if SUPPORTS_ARKIT_2_0
+	if (FAppleARKitAvailability::SupportsARKit20())
+	{
+		WorldMappingState = ToEARWorldMappingState(InARFrame.worldMappingStatus);
+	}
+#endif
 }
 
 FAppleARKitFrame::FAppleARKitFrame( const FAppleARKitFrame& Other )
@@ -89,6 +116,7 @@ FAppleARKitFrame::FAppleARKitFrame( const FAppleARKitFrame& Other )
   , CapturedCbCrImageHeight( Other.CapturedCbCrImageHeight )
   , Camera( Other.Camera )
   , LightEstimate( Other.LightEstimate )
+	, WorldMappingState(Other.WorldMappingState)
 {
 	if(Other.NativeFrame != nullptr)
 	{
diff --git a/Engine/Plugins/Runtime/AppleARKit/Source/AppleARKit/Private/AppleARKitSystem.cpp b/Engine/Plugins/Runtime/AppleARKit/Source/AppleARKit/Private/AppleARKitSystem.cpp
--- a/Engine/Plugins/Runtime/AppleARKit/Source/AppleARKit/Private/AppleARKitSystem.cpp
+++ b/Engine/Plugins/Runtime/AppleARKit/Source/AppleARKit/Private/AppleARKitSystem.cpp
@@ -14,7 +14,8 @@
 #include "GeneralProjectSettings.h"
 #include "ARSessionConfig.h"
 #include "AppleARKitSettings.h"
-#include "ARTrackable.h"
+//@joeg -- Environment capture support
+#include "AppleARKitTrackable.h"
 #include "ARLightEstimate.h"
 #include "ARTraceResult.h"
 #include "ARPin.h"
@@ -142,7 +143,6 @@ private:
 	FAppleARKitVideoOverlay VideoOverlay;
 };
 
-
 //
 //  FAppleARKitSystem
 //
@@ -725,6 +725,241 @@ bool FAppleARKitSystem::OnIsTrackingTypeSupported(EARSessionType SessionType) co
 	return false;
 }
 
+//@joeg -- ARKit 2.0 additions
+
+bool FAppleARKitSystem::OnAddManualEnvironmentCaptureProbe(FVector Location, FVector Extent)
+{
+#if SUPPORTS_ARKIT_2_0
+	if (Session != nullptr)
+	{
+		if (FAppleARKitAvailability::SupportsARKit20())
+		{
+//@joeg -- Todo need to fix this transform as it needs to use the alignment transform too
+			// Build and add the anchor
+			simd_float4x4 AnchorMatrix = FAppleARKitConversion::ToARKitMatrix(FTransform(Location));
+			simd_float3 AnchorExtent = FAppleARKitConversion::ToARKitVector(Extent * 2.f);
+			AREnvironmentProbeAnchor* ARProbe = [[AREnvironmentProbeAnchor alloc] initWithTransform: AnchorMatrix extent: AnchorExtent];
+			[Session addAnchor: ARProbe];
+			[ARProbe release];
+		}
+		return true;
+	}
+#endif
+	return false;
+}
+
+#if SUPPORTS_ARKIT_2_0
+/** Since both the object extraction and world saving need to get the world map async, use a common chunk of code for this */
+class FAppleARKitGetWorldMapObjectAsyncTask
+{
+public:
+	/** Performs the call to get the world map and triggers the OnWorldMapAcquired() the completion handler */
+	void Run()
+	{
+		[Session getCurrentWorldMapWithCompletionHandler: ^(ARWorldMap* worldMap, NSError* error)
+		 {
+			 WorldMap = worldMap;
+			 [WorldMap retain];
+			 bool bWasSuccessful = error == nullptr;
+			 FString ErrorString;
+			 if (error != nullptr)
+			 {
+				 ErrorString = [error localizedDescription];
+			 }
+			 OnWorldMapAcquired(bWasSuccessful, ErrorString);
+		 }];
+	}
+	
+protected:
+	FAppleARKitGetWorldMapObjectAsyncTask(ARSession* InSession) :
+		Session(InSession)
+	{
+		CFRetain(Session);
+	}
+	
+	void Release()
+	{
+		if (Session != nullptr)
+		{
+			[Session release];
+			Session = nullptr;
+		}
+		if (WorldMap != nullptr)
+		{
+			[WorldMap release];
+			WorldMap = nullptr;
+		}
+	}
+
+	/** Called once the world map completion handler is called */
+	virtual void OnWorldMapAcquired(bool bWasSuccessful, FString ErrorString) = 0;
+
+	/** The session object that we'll grab the world from */
+	ARSession* Session;
+	/** The world map object once the call has completed */
+	ARWorldMap* WorldMap;
+};
+
+//@joeg -- The API changed last minute so you don't need to resolve the world to get an object anymore
+// This needs to be cleaned up
+class FAppleARKitGetCandidateObjectAsyncTask :
+	public FARGetCandidateObjectAsyncTask
+{
+public:
+	FAppleARKitGetCandidateObjectAsyncTask(ARSession* InSession, FVector InLocation, FVector InExtent) :
+		Location(InLocation)
+		, Extent(InExtent)
+		, ReferenceObject(nullptr)
+		, Session(InSession)
+	{
+		[Session retain];
+	}
+	
+	/** @return the candidate object that you can use for detection later */
+	virtual UARCandidateObject* GetCandidateObject() override
+	{
+		if (ReferenceObject != nullptr)
+		{
+			UARCandidateObject* CandidateObject = NewObject<UARCandidateObject>();
+			
+			FVector RefObjCenter = FAppleARKitConversion::ToFVector(ReferenceObject.center);
+			FVector RefObjExtent = 0.5f * FAppleARKitConversion::ToFVector(ReferenceObject.extent);
+			FBox BoundingBox(RefObjCenter, RefObjExtent);
+			CandidateObject->SetBoundingBox(BoundingBox);
+			
+			// Serialize the object into a byte array and stick that on the candidate object
+			NSData* RefObjData = [NSKeyedArchiver archivedDataWithRootObject: ReferenceObject];
+			uint32 SavedSize = RefObjData.length;
+			TArray<uint8> RawBytes;
+			RawBytes.AddUninitialized(SavedSize);
+			FPlatformMemory::Memcpy(RawBytes.GetData(), [RefObjData bytes], SavedSize);
+			CandidateObject->SetCandidateObjectData(RawBytes);
+
+			return CandidateObject;
+		}
+		return nullptr;
+	}
+	
+	virtual ~FAppleARKitGetCandidateObjectAsyncTask()
+	{
+		[Session release];
+		if (ReferenceObject != nullptr)
+		{
+			CFRelease(ReferenceObject);
+		}
+	}
+
+	/** Performs the call to get the world map and triggers the OnWorldMapAcquired() the completion handler */
+	void Run()
+	{
+		simd_float4x4 ARMatrix = FAppleARKitConversion::ToARKitMatrix(FTransform(Location));
+		simd_float3 Center = 0.f;
+		simd_float3 ARExtent = FAppleARKitConversion::ToARKitVector(Extent * 2.f);
+
+		[Session createReferenceObjectWithTransform: ARMatrix center: Center extent: ARExtent
+		 completionHandler: ^(ARReferenceObject* refObject, NSError* error)
+		{
+			ReferenceObject = refObject;
+			bool bWasSuccessful = error == nullptr;
+			bHadError = error != nullptr;
+			FString ErrorString;
+			if (error != nullptr)
+			{
+				ErrorString = [error localizedDescription];
+			}
+			bIsDone = true;
+		}];
+	}
+	
+private:
+	FVector Location;
+	FVector Extent;
+	ARReferenceObject* ReferenceObject;
+
+	/** The session object that we'll grab the object from */
+	ARSession* Session;
+};
+
+class FAppleARKitSaveWorldAsyncTask :
+	public FARSaveWorldAsyncTask,
+	public FAppleARKitGetWorldMapObjectAsyncTask
+{
+public:
+	FAppleARKitSaveWorldAsyncTask(ARSession* InSession) :
+		FAppleARKitGetWorldMapObjectAsyncTask(InSession)
+	{
+	}
+
+	virtual ~FAppleARKitSaveWorldAsyncTask()
+	{
+		Release();
+	}
+
+	virtual void OnWorldMapAcquired(bool bWasSuccessful, FString ErrorString) override
+	{
+		if (bWasSuccessful)
+		{
+			NSData* WorldNSData = [NSKeyedArchiver archivedDataWithRootObject: WorldMap];
+
+			// Copy to our TArray that will serve the data to the caller
+			uint32 SavedSize = WorldNSData.length;
+			WorldData.AddUninitialized(SavedSize);
+			FPlatformMemory::Memcpy(WorldData.GetData(), [WorldNSData bytes], SavedSize);
+		}
+		else
+		{
+			Error = ErrorString;
+			bHadError = true;
+		}
+		// Trigger that we're done
+		bIsDone = true;
+	}
+};
+#endif
+
+TSharedPtr<FARGetCandidateObjectAsyncTask, ESPMode::ThreadSafe> FAppleARKitSystem::OnGetCandidateObject(FVector Location, FVector Extent) const
+{
+#if SUPPORTS_ARKIT_2_0
+	if (Session != nullptr)
+	{
+		if (FAppleARKitAvailability::SupportsARKit20())
+		{
+			TSharedPtr<FAppleARKitGetCandidateObjectAsyncTask, ESPMode::ThreadSafe> Task = MakeShared<FAppleARKitGetCandidateObjectAsyncTask, ESPMode::ThreadSafe>(Session, Location, Extent);
+			Task->Run();
+			return Task;
+		}
+	}
+#endif
+	return  MakeShared<FARErrorGetCandidateObjectAsyncTask, ESPMode::ThreadSafe>(TEXT("GetCandidateObject - requires a valid, running ARKit 2.0 session"));
+}
+
+TSharedPtr<FARSaveWorldAsyncTask, ESPMode::ThreadSafe> FAppleARKitSystem::OnSaveWorld() const
+{
+#if SUPPORTS_ARKIT_2_0
+	if (Session != nullptr)
+	{
+		if (FAppleARKitAvailability::SupportsARKit20())
+		{
+			TSharedPtr<FAppleARKitSaveWorldAsyncTask, ESPMode::ThreadSafe> Task = MakeShared<FAppleARKitSaveWorldAsyncTask, ESPMode::ThreadSafe>(Session);
+			Task->Run();
+			return Task;
+		}
+	}
+#endif
+	return  MakeShared<FARErrorSaveWorldAsyncTask, ESPMode::ThreadSafe>(TEXT("SaveWorld - requires a valid, running ARKit 2.0 session"));
+}
+
+EARWorldMappingState FAppleARKitSystem::OnGetWorldMappingStatus() const
+{
+	if (GameThreadFrame.IsValid())
+	{
+		return GameThreadFrame->WorldMappingState;
+	}
+	return EARWorldMappingState::NotAvailable;
+}
+
+//@joeg -- End additions
+
 void FAppleARKitSystem::AddReferencedObjects( FReferenceCollector& Collector )
 {
 	FARSystemBase::AddReferencedObjects(Collector);
@@ -734,6 +969,8 @@ void FAppleARKitSystem::AddReferencedObjects( FReferenceCollector& Collector )
 	Collector.AddReferencedObject( CameraImage );
 	Collector.AddReferencedObject( CameraDepth );
 	Collector.AddReferencedObjects( CandidateImages );
+//@joeg -- Object detection
+	Collector.AddReferencedObjects( CandidateObjects );
 
 	if(LightEstimate)
 	{
@@ -828,6 +1065,14 @@ bool FAppleARKitSystem::Run(UARSessionConfig* SessionConfig)
 		return true;
 	}
 
+//@joeg -- bug fix
+	{
+		// Clear out any existing frames since they aren't valid anymore
+		FScopeLock ScopeLock(&FrameLock);
+		GameThreadFrame = TSharedPtr<FAppleARKitFrame, ESPMode::ThreadSafe>();
+		LastReceivedFrame = TSharedPtr<FAppleARKitFrame, ESPMode::ThreadSafe>();
+	}
+	
 #if SUPPORTS_ARKIT_1_0
 	if (FAppleARKitAvailability::SupportsARKit10())
 	{
@@ -837,7 +1082,7 @@ bool FAppleARKitSystem::Run(UARSessionConfig* SessionConfig)
 		CheckForFaceARSupport(SessionConfig);
 		if (FaceARSupport == nullptr)
 		{
-			Configuration = FAppleARKitConversion::ToARConfiguration(SessionConfig, CandidateImages, ConvertedCandidateImages);
+			Configuration = FAppleARKitConversion::ToARConfiguration(SessionConfig, CandidateImages, ConvertedCandidateImages, CandidateObjects);
 		}
 		else
 		{
@@ -1004,9 +1249,35 @@ static TSharedPtr<FAppleARKitAnchorData> MakeAnchorData( ARAnchor* Anchor )
 		NewAnchor = MakeShared<FAppleARKitAnchorData>(
 			FAppleARKitConversion::ToFGuid(ImageAnchor.identifier),
 			FAppleARKitConversion::ToFTransform(ImageAnchor.transform),
+//@joeg -- Changed to share image and object detection
+			EAppleAnchorType::ImageAnchor,
 			FString(ImageAnchor.referenceImage.name)
 		);
 	}
+#endif
+//@joeg -- Added environmental texture probe support
+#if SUPPORTS_ARKIT_2_0
+	else if (FAppleARKitAvailability::SupportsARKit20() && [Anchor isKindOfClass:[AREnvironmentProbeAnchor class]])
+	{
+		AREnvironmentProbeAnchor* ProbeAnchor = (AREnvironmentProbeAnchor*)Anchor;
+		NewAnchor = MakeShared<FAppleARKitAnchorData>(
+			FAppleARKitConversion::ToFGuid(ProbeAnchor.identifier),
+			FAppleARKitConversion::ToFTransform(ProbeAnchor.transform),
+			0.5f * FAppleARKitConversion::ToFVector(ProbeAnchor.extent).GetAbs(),
+			ProbeAnchor.environmentTexture
+		);
+	}
+//@joeg -- Object dectection support
+	else if (FAppleARKitAvailability::SupportsARKit20() && [Anchor isKindOfClass:[ARObjectAnchor class]])
+	{
+		ARObjectAnchor* ObjectAnchor = (ARObjectAnchor*)Anchor;
+		NewAnchor = MakeShared<FAppleARKitAnchorData>(
+			  FAppleARKitConversion::ToFGuid(ObjectAnchor.identifier),
+			  FAppleARKitConversion::ToFTransform(ObjectAnchor.transform),
+			  EAppleAnchorType::ObjectAnchor,
+			  FString(ObjectAnchor.referenceObject.name)
+		  );
+	}
 #endif
 	else
 	{
@@ -1128,27 +1399,43 @@ void FAppleARKitSystem::SessionDidAddAnchors_Internal( TSharedRef<FAppleARKitAnc
 			NewGeometry = NewGeo;
 			break;
 		}
+		case EAppleAnchorType::FaceAnchor:
+		{
+			NewAnchorDebugName = FString::Printf(TEXT("FACE-%02d"), LastTrackedGeometry_DebugId++);
+			UARFaceGeometry* NewGeo = NewObject<UARFaceGeometry>();
+//@joeg -- Eye tracking support
+			NewGeo->UpdateFaceGeometry(SharedThis(this), GameThreadFrameNumber, GameThreadTimestamp, AnchorData->Transform, GetAlignmentTransform(), AnchorData->BlendShapes, AnchorData->FaceVerts, AnchorData->FaceIndices, AnchorData->LeftEyeTransform, AnchorData->RightEyeTransform, AnchorData->LookAtTarget);
+			NewGeometry = NewGeo;
+			break;
+		}
 		case EAppleAnchorType::ImageAnchor:
 		{
 			NewAnchorDebugName = FString::Printf(TEXT("IMG-%02d"), LastTrackedGeometry_DebugId++);
 			UARTrackedImage* NewImage = NewObject<UARTrackedImage>();
-			UARCandidateImage** CandidateImage = CandidateImages.Find(AnchorData->DetectedImageName);
+			UARCandidateImage** CandidateImage = CandidateImages.Find(AnchorData->DetectedAnchorName);
 			ensure(CandidateImage != nullptr);
 			NewImage->UpdateTrackedGeometry(SharedThis(this), GameThreadFrameNumber, GameThreadTimestamp, AnchorData->Transform, GetAlignmentTransform(), *CandidateImage);
 			NewGeometry = NewImage;
 			break;
 		}
-		case EAppleAnchorType::FaceAnchor:
+//@joeg -- Added environmental texture probe support
+		case EAppleAnchorType::EnvironmentProbeAnchor:
 		{
-			UARFaceGeometry* NewGeo = NewObject<UARFaceGeometry>();
-			NewAnchorDebugName = FString::Printf(TEXT("FACE-%02d"), LastTrackedGeometry_DebugId++);
-			NewGeo->SetDebugName(FName(*NewAnchorDebugName));
-			NewGeo->UpdateTrackedGeometry(SharedThis(this), GameThreadFrameNumber, GameThreadTimestamp, AnchorData->Transform, GetAlignmentTransform(), AnchorData->BlendShapes, AnchorData->FaceVerts, AnchorData->FaceIndices);
-			NewGeometry = NewGeo;
+			NewAnchorDebugName = FString::Printf(TEXT("ENV-%02d"), LastTrackedGeometry_DebugId++);
+			UAppleARKitEnvironmentCaptureProbe* NewProbe = NewObject<UAppleARKitEnvironmentCaptureProbe>();
+			NewProbe->UpdateEnvironmentCapture(SharedThis(this), GameThreadFrameNumber, GameThreadTimestamp, AnchorData->Transform, GetAlignmentTransform(), AnchorData->Extent, AnchorData->ProbeTexture);
+			NewGeometry = NewProbe;
+			break;
 		}
-		default:
+//@joeg -- Object detection
+		case EAppleAnchorType::ObjectAnchor:
 		{
-			NewAnchorDebugName = FString::Printf(TEXT("APPL-%02d"), LastTrackedGeometry_DebugId++);
+			NewAnchorDebugName = FString::Printf(TEXT("OBJ-%02d"), LastTrackedGeometry_DebugId++);
+			UARTrackedObject* NewTrackedObject = NewObject<UARTrackedObject>();
+			UARCandidateObject** CandidateObject = CandidateObjects.Find(AnchorData->DetectedAnchorName);
+			ensure(CandidateObject != nullptr);
+			NewTrackedObject->UpdateTrackedGeometry(SharedThis(this), GameThreadFrameNumber, GameThreadTimestamp, AnchorData->Transform, GetAlignmentTransform(), *CandidateObject);
+			NewGeometry = NewTrackedObject;
 			break;
 		}
 	}
@@ -1217,7 +1504,41 @@ void FAppleARKitSystem::SessionDidUpdateAnchors_Internal( TSharedRef<FAppleARKit
 			{
 				if (UARFaceGeometry* FaceGeo = Cast<UARFaceGeometry>(FoundGeometry))
 				{
-					FaceGeo->UpdateTrackedGeometry(SharedThis(this), GameThreadFrameNumber, GameThreadTimestamp, AnchorData->Transform, GetAlignmentTransform(), AnchorData->BlendShapes, AnchorData->FaceVerts, AnchorData->FaceIndices);
+//@joeg -- Eye tracking support
+					FaceGeo->UpdateFaceGeometry(SharedThis(this), GameThreadFrameNumber, GameThreadTimestamp, AnchorData->Transform, GetAlignmentTransform(), AnchorData->BlendShapes, AnchorData->FaceVerts, AnchorData->FaceIndices, AnchorData->LeftEyeTransform, AnchorData->RightEyeTransform, AnchorData->LookAtTarget);
+					for (UARPin* Pin : PinsToUpdate)
+					{
+						const FTransform Pin_LocalToTrackingTransform_PostUpdate = Pin->GetLocalToTrackingTransform_NoAlignment() * AnchorDeltaTransform;
+						Pin->OnTransformUpdated(Pin_LocalToTrackingTransform_PostUpdate);
+					}
+				}
+				break;
+			}
+//@joeg -- Added image tracking support
+            case EAppleAnchorType::ImageAnchor:
+            {
+				if (UARTrackedImage* ImageAnchor = Cast<UARTrackedImage>(FoundGeometry))
+				{
+//@joeg -- Changed so object and image detection can share
+					UARCandidateImage** CandidateImage = CandidateImages.Find(AnchorData->DetectedAnchorName);
+					ensure(CandidateImage != nullptr);
+
+					ImageAnchor->UpdateTrackedGeometry(SharedThis(this), GameThreadFrameNumber, GameThreadTimestamp, AnchorData->Transform, GetAlignmentTransform(), *CandidateImage);
+					for (UARPin* Pin : PinsToUpdate)
+					{
+						const FTransform Pin_LocalToTrackingTransform_PostUpdate = Pin->GetLocalToTrackingTransform_NoAlignment() * AnchorDeltaTransform;
+						Pin->OnTransformUpdated(Pin_LocalToTrackingTransform_PostUpdate);
+					}
+				}
+                break;
+            }
+//@joeg -- Added environmental texture probe support
+			case EAppleAnchorType::EnvironmentProbeAnchor:
+			{
+				if (UAppleARKitEnvironmentCaptureProbe* ProbeAnchor = Cast<UAppleARKitEnvironmentCaptureProbe>(FoundGeometry))
+				{
+					// NOTE: The metal texture will be a different texture every time the cubemap is updated which requires a render resource flush
+					ProbeAnchor->UpdateEnvironmentCapture(SharedThis(this), GameThreadFrameNumber, GameThreadTimestamp, AnchorData->Transform, GetAlignmentTransform(), AnchorData->Extent, AnchorData->ProbeTexture);
 					for (UARPin* Pin : PinsToUpdate)
 					{
 						const FTransform Pin_LocalToTrackingTransform_PostUpdate = Pin->GetLocalToTrackingTransform_NoAlignment() * AnchorDeltaTransform;
diff --git a/Engine/Plugins/Runtime/AppleARKit/Source/AppleARKit/Private/AppleARKitSystem.h b/Engine/Plugins/Runtime/AppleARKit/Source/AppleARKit/Private/AppleARKitSystem.h
--- a/Engine/Plugins/Runtime/AppleARKit/Source/AppleARKit/Private/AppleARKitSystem.h
+++ b/Engine/Plugins/Runtime/AppleARKit/Source/AppleARKit/Private/AppleARKitSystem.h
@@ -79,6 +79,12 @@ protected:
 	virtual void OnRemovePin(UARPin* PinToRemove) override;
 	virtual UARTextureCameraImage* OnGetCameraImage() override;
 	virtual UARTextureCameraDepth* OnGetCameraDepth() override;
+//@joeg -- ARKit 2.0 additions
+	virtual bool OnAddManualEnvironmentCaptureProbe(FVector Location, FVector Extent) override;
+	virtual TSharedPtr<FARGetCandidateObjectAsyncTask, ESPMode::ThreadSafe> OnGetCandidateObject(FVector Location, FVector Extent) const override;
+	virtual TSharedPtr<FARSaveWorldAsyncTask, ESPMode::ThreadSafe> OnSaveWorld() const override;
+	virtual EARWorldMappingState OnGetWorldMappingStatus() const override;
+//@joeg -- End additions
 	//~IARSystemSupport
 
 private:
@@ -161,6 +167,8 @@ private:
 	UAppleARKitTextureCameraImage* CameraImage;
 	UAppleARKitTextureCameraDepth* CameraDepth;
 	TMap< FString, UARCandidateImage* > CandidateImages;
+//@joeg -- Object detection
+	TMap< FString, UARCandidateObject* > CandidateObjects;
 	// ...
 	// PROPERTIES REPORTED TO FGCObject
 	//
diff --git a/Engine/Plugins/Runtime/AppleARKit/Source/AppleARKit/Private/AppleARKitTextures.cpp b/Engine/Plugins/Runtime/AppleARKit/Source/AppleARKit/Private/AppleARKitTextures.cpp
--- a/Engine/Plugins/Runtime/AppleARKit/Source/AppleARKit/Private/AppleARKitTextures.cpp
+++ b/Engine/Plugins/Runtime/AppleARKit/Source/AppleARKit/Private/AppleARKitTextures.cpp
@@ -1,6 +1,13 @@
 // Copyright 1998-2018 Epic Games, Inc. All Rights Reserved.
 
 #include "AppleARKitTextures.h"
+#include "ExternalTexture.h"
+#include "RenderingThread.h"
+#include "Containers/DynamicRHIResourceArray.h"
+
+#if PLATFORM_MAC || PLATFORM_IOS
+	#import <Metal/Metal.h>
+#endif
 
 UAppleARKitTextureCameraImage::UAppleARKitTextureCameraImage(const FObjectInitializer& ObjectInitializer)
 	: Super(ObjectInitializer)
@@ -88,3 +95,254 @@ void UAppleARKitTextureCameraDepth::Init(float InTimestamp, AVDepthData* InCamer
 }
 
 #endif
+
+//@joeg -- Added for environment capture support
+
+UAppleARKitEnvironmentCaptureProbeTexture::UAppleARKitEnvironmentCaptureProbeTexture(const FObjectInitializer& ObjectInitializer)
+	: Super(ObjectInitializer)
+#if PLATFORM_MAC || PLATFORM_IOS
+	, MetalTexture(nullptr)
+#endif
+{
+	ExternalTextureGuid = FGuid::NewGuid();
+	SRGB = false;
+}
+
+#if PLATFORM_MAC || PLATFORM_IOS
+void UAppleARKitEnvironmentCaptureProbeTexture::Init(float InTimestamp, id<MTLTexture> InEnvironmentTexture)
+{
+	if (Resource == nullptr)
+	{
+		UpdateResource();
+	}
+	
+	// Do nothing if the textures are the same
+	// They will change as the data comes in but the textures themselves may stay the same between updates
+	if (MetalTexture == InEnvironmentTexture)
+	{
+		return;
+	}
+	
+	// Handle the case where this UObject is being reused
+	if (MetalTexture != nullptr)
+	{
+		CFRelease(MetalTexture);
+		MetalTexture = nullptr;
+	}
+	
+	if (InEnvironmentTexture != nullptr)
+	{
+		Timestamp = InTimestamp;
+		MetalTexture = InEnvironmentTexture;
+		CFRetain(MetalTexture);
+		Size.X = MetalTexture.width;
+		Size.Y = MetalTexture.height;
+	}
+	// Force an update to our external texture on the render thread
+	if (Resource != nullptr)
+	{
+		ENQUEUE_UNIQUE_RENDER_COMMAND_ONEPARAMETER(UpdateEnvironmentCapture,
+		   FTextureResource*, InResource, Resource,
+		   {
+			   InResource->InitRHI();
+		   });
+	}
+}
+
+/**
+ * Passes a metaltexture through to the RHI to wrap in an RHI texture without traversing system memory.
+ */
+class FAppleARKitMetalTextureResourceWrapper :
+	public FResourceBulkDataInterface
+{
+public:
+	FAppleARKitMetalTextureResourceWrapper(id<MTLTexture> InImageBuffer)
+		: ImageBuffer(InImageBuffer)
+	{
+		check(ImageBuffer);
+		CFRetain(ImageBuffer);
+	}
+	
+	virtual ~FAppleARKitMetalTextureResourceWrapper()
+	{
+		CFRelease(ImageBuffer);
+		ImageBuffer = nullptr;
+	}
+
+	/**
+	 * @return ptr to the resource memory which has been preallocated
+	 */
+	virtual const void* GetResourceBulkData() const override
+	{
+		return ImageBuffer;
+	}
+	
+	/**
+	 * @return size of resource memory
+	 */
+	virtual uint32 GetResourceBulkDataSize() const override
+	{
+		return 0;
+	}
+	
+	/**
+	 * @return the type of bulk data for special handling
+	 */
+	virtual EBulkDataType GetResourceType() const override
+	{
+		return EBulkDataType::MediaTexture;
+	}
+	
+	/**
+	 * Free memory after it has been used to initialize RHI resource
+	 */
+	virtual void Discard() override
+	{
+		delete this;
+	}
+	
+	id<MTLTexture> ImageBuffer;
+};
+
+class FARMetalResource :
+	public FTextureResource
+{
+public:
+	FARMetalResource(UAppleARKitEnvironmentCaptureProbeTexture* InOwner)
+		: Owner(InOwner)
+	{
+		bGreyScaleFormat = false;
+		bSRGB = true;
+	}
+	
+	virtual ~FARMetalResource()
+	{
+	}
+	
+	/**
+	 * Called when the resource is initialized. This is only called by the rendering thread.
+	 */
+	virtual void InitRHI() override
+	{
+		FRHIResourceCreateInfo CreateInfo;
+
+		id<MTLTexture> MetalTexture = Owner->GetMetalTexture();
+		if (MetalTexture != nullptr)
+		{
+			Size.X = Size.Y = Owner->Size.X;
+			const int32 NumComponents = 4;
+			const int32 BytesPerImage = Size.X * Size.Y * NumComponents;
+			const int32 BytesPerRow = Size.X * NumComponents;
+
+			const uint32 CreateFlags = TexCreate_NoMipTail | TexCreate_SRGB;
+			EnvCubemapTextureRHIRef = RHICreateTextureCube(Size.X, PF_B8G8R8A8, 1, CreateFlags, CreateInfo);
+			/** To map their texture faces into our space we need:
+				-Z to +X CCW 90
+				+Z to -X CW 90
+				+X to +Y CCW 90
+				-X to -Y CCW 90
+				+Y to +Z CW 180
+				-Y to -Z
+			 */
+			CopyCubeFace(MetalTexture, EnvCubemapTextureRHIRef, kCGImagePropertyOrientationLeft,  5, 0);
+			CopyCubeFace(MetalTexture, EnvCubemapTextureRHIRef, kCGImagePropertyOrientationRight, 4, 1);
+			CopyCubeFace(MetalTexture, EnvCubemapTextureRHIRef, kCGImagePropertyOrientationLeft, 0, 2);
+			CopyCubeFace(MetalTexture, EnvCubemapTextureRHIRef, kCGImagePropertyOrientationLeft, 1, 3);
+			CopyCubeFace(MetalTexture, EnvCubemapTextureRHIRef, kCGImagePropertyOrientationDown, 2, 4);
+			CopyCubeFace(MetalTexture, EnvCubemapTextureRHIRef, kCGImagePropertyOrientationUp, 3, 5);
+		}
+		else
+		{
+			Size.X = Size.Y = 1;
+			// Start with a 1x1 texture
+			EnvCubemapTextureRHIRef = RHICreateTextureCube(Size.X, PF_B8G8R8A8, 1, 0, CreateInfo);
+		}
+
+		TextureRHI = EnvCubemapTextureRHIRef;
+		TextureRHI->SetName(Owner->GetFName());
+		RHIBindDebugLabelName(TextureRHI, *Owner->GetName());
+		RHIUpdateTextureReference(Owner->TextureReference.TextureReferenceRHI, TextureRHI);
+
+		FSamplerStateInitializerRHI SamplerStateInitializer(SF_Bilinear, AM_Clamp, AM_Clamp, AM_Clamp);
+		SamplerStateRHI = RHICreateSamplerState(SamplerStateInitializer);
+	}
+
+//@joeg -- Move to image utils plugin
+	void CopyCubeFace(id<MTLTexture> MetalTexture, FTextureCubeRHIRef Cubemap, uint32 Rotation, int32 MetalCubeIndex, int32 OurCubeIndex)
+	{
+		// Rotate the image we need to get a view into the face as a new slice
+		id<MTLTexture> CubeFaceMetalTexture = [MetalTexture newTextureViewWithPixelFormat: MTLPixelFormatBGRA8Unorm textureType: MTLTextureType2D levels: NSMakeRange(0, 1) slices: NSMakeRange(MetalCubeIndex, 1)];
+		CIImage* CubefaceImage = [[CIImage alloc] initWithMTLTexture: CubeFaceMetalTexture options: nil];
+		CIImage* RotatedCubefaceImage = [CubefaceImage imageByApplyingOrientation: Rotation];
+		CIImage* ImageTransform = nullptr;
+		if (Rotation != kCGImagePropertyOrientationUp)
+		{
+			ImageTransform = RotatedCubefaceImage;
+		}
+		else
+		{
+			// We don't need to rotate it so just use a copy instead
+			ImageTransform = CubefaceImage;
+		}
+
+		// Make a new view into our texture and directly render to that to avoid the CPU copy
+		id<MTLTexture> OurCubeFaceMetalTexture = [(id<MTLTexture>)Cubemap->GetNativeResource() newTextureViewWithPixelFormat: MTLPixelFormatBGRA8Unorm textureType: MTLTextureType2D levels: NSMakeRange(0, 1) slices: NSMakeRange(OurCubeIndex, 1)];
+
+		CIContext* Context = [CIContext context];
+
+		[Context render: RotatedCubefaceImage toMTLTexture: OurCubeFaceMetalTexture commandBuffer: nil bounds: CubefaceImage.extent colorSpace: CubefaceImage.colorSpace];
+		[CubefaceImage release];
+//			[RotatedCubefaceImage release];
+		[CubeFaceMetalTexture release];
+		[OurCubeFaceMetalTexture release];
+	}
+	
+	virtual void ReleaseRHI() override
+	{
+		RHIUpdateTextureReference(Owner->TextureReference.TextureReferenceRHI, FTextureRHIParamRef());
+		EnvCubemapTextureRHIRef.SafeRelease();
+		FTextureResource::ReleaseRHI();
+		FExternalTextureRegistry::Get().UnregisterExternalTexture(Owner->ExternalTextureGuid);
+	}
+	
+	/** Returns the width of the texture in pixels. */
+	virtual uint32 GetSizeX() const override
+	{
+		return Size.X;
+	}
+	
+	/** Returns the height of the texture in pixels. */
+	virtual uint32 GetSizeY() const override
+	{
+		return Size.Y;
+	}
+	
+private:
+	FIntPoint Size;
+	
+	FTextureCubeRHIRef EnvCubemapTextureRHIRef;
+	
+	const UAppleARKitEnvironmentCaptureProbeTexture* Owner;
+};
+
+#endif
+
+FTextureResource* UAppleARKitEnvironmentCaptureProbeTexture::CreateResource()
+{
+#if PLATFORM_MAC || PLATFORM_IOS
+	return new FARMetalResource(this);
+#endif
+	return nullptr;
+}
+
+void UAppleARKitEnvironmentCaptureProbeTexture::BeginDestroy()
+{
+#if PLATFORM_MAC || PLATFORM_IOS
+	if (MetalTexture != nullptr)
+	{
+		CFRelease(MetalTexture);
+		MetalTexture = nullptr;
+	}
+#endif
+	Super::BeginDestroy();
+}
diff --git a/Engine/Plugins/Runtime/AppleARKit/Source/AppleARKit/Private/AppleARKitTrackable.cpp b/Engine/Plugins/Runtime/AppleARKit/Source/AppleARKit/Private/AppleARKitTrackable.cpp
new file mode 100644
--- /dev/null
+++ b/Engine/Plugins/Runtime/AppleARKit/Source/AppleARKit/Private/AppleARKitTrackable.cpp
@@ -0,0 +1,24 @@
+// Copyright 1998-2017 Epic Games, Inc. All Rights Reserved.
+
+#include "AppleARKitTrackable.h"
+#include "AppleARKitTextures.h"
+
+UAppleARKitEnvironmentCaptureProbe::UAppleARKitEnvironmentCaptureProbe()
+	: Super()
+	, ARKitEnvironmentTexture(nullptr)
+{
+	ARKitEnvironmentTexture = NewObject<UAppleARKitEnvironmentCaptureProbeTexture>();
+	// Set the base class member since that's what gets used by the non-ARKit specific code
+	EnvironmentCaptureTexture = ARKitEnvironmentTexture;
+}
+
+#if PLATFORM_MAC || PLATFORM_IOS
+void UAppleARKitEnvironmentCaptureProbe::UpdateEnvironmentCapture(const TSharedRef<FARSystemBase, ESPMode::ThreadSafe>& InTrackingSystem, uint32 FrameNumber, double InTimestamp, const FTransform& InLocalToTrackingTransform, const FTransform& InAlignmentTransform, FVector InExtent, id<MTLTexture> InMetalTexture)
+{
+	Super::UpdateEnvironmentCapture(InTrackingSystem, FrameNumber, InTimestamp, InLocalToTrackingTransform, InAlignmentTransform, InExtent);
+	
+	check(ARKitEnvironmentTexture != nullptr);
+	ARKitEnvironmentTexture->Init(InTimestamp, InMetalTexture);
+}
+#endif
+
diff --git a/Engine/Plugins/Runtime/AppleARKit/Source/AppleARKit/Public/AppleARKitAvailability.h b/Engine/Plugins/Runtime/AppleARKit/Source/AppleARKit/Public/AppleARKitAvailability.h
--- a/Engine/Plugins/Runtime/AppleARKit/Source/AppleARKit/Public/AppleARKitAvailability.h
+++ b/Engine/Plugins/Runtime/AppleARKit/Source/AppleARKit/Public/AppleARKitAvailability.h
@@ -2,12 +2,11 @@
 
 #pragma once
 
-#ifndef PLATFORM_IOS
-	#define PLATFORM_IOS 0
-#endif
 
-#ifndef PLATFORM_TVOS
-	#define PLATFORM_TVOS 0
+//@joeg -- Fix for non-unity build error
+#include "CoreMinimal.h"
+#if PLATFORM_IOS && !PLATFORM_TVOS
+	#include "Availability.h"
 #endif
 
 #if PLATFORM_IOS && !PLATFORM_TVOS
@@ -26,11 +25,22 @@
 		#define SUPPORTS_ARKIT_1_5 0
 	#endif
 
+//@joeg -- Added image tracking support
+	// Check for ARKit 2.0
+	#ifdef __IPHONE_12_0
+//	#ifdef __IPHONE_11_3
+		#define SUPPORTS_ARKIT_2_0 1
+	#else
+		#define SUPPORTS_ARKIT_2_0 0
+	#endif
+
 #else
 
 	// No ARKit support
 	#define SUPPORTS_ARKIT_1_0 0
 	#define SUPPORTS_ARKIT_1_5 0
+//@joeg -- Added image tracking support
+	#define SUPPORTS_ARKIT_2_0 0
 
 #endif
 
@@ -72,4 +82,24 @@ public:
 #endif
 		return bSupportsARKit15;
 	}
+
+//@joeg -- Added image tracking support
+	static bool SupportsARKit20()
+	{
+		static bool bSupportsARKit20 = false;
+#if SUPPORTS_ARKIT_2_0
+		static bool bSupportChecked = false;
+		if (!bSupportChecked)
+		{
+			bSupportChecked = true;
+			// This call is slow, so cache the result
+			if (@available(iOS 12.0, *))
+//			if (@available(iOS 11.3, *))
+			{
+				bSupportsARKit20 = true;
+			}
+		}
+#endif
+		return bSupportsARKit20;
+	}
 };
diff --git a/Engine/Plugins/Runtime/AppleARKit/Source/AppleARKit/Public/AppleARKitConversion.h b/Engine/Plugins/Runtime/AppleARKit/Source/AppleARKit/Public/AppleARKitConversion.h
--- a/Engine/Plugins/Runtime/AppleARKit/Source/AppleARKit/Public/AppleARKitConversion.h
+++ b/Engine/Plugins/Runtime/AppleARKit/Source/AppleARKit/Public/AppleARKitConversion.h
@@ -2,7 +2,6 @@
 
 #pragma once
 
-// UE4
 #include "Math/Transform.h"
 #include "ARPin.h"
 
@@ -46,7 +45,41 @@ struct FAppleARKitConversion
 		return FTransform( Rotation, Translation );
 	}
 
-	/** 
+    /**
+     * Convert's an Unreal 'Z up' transform to ARKit's 'Y up' 'right handed' coordinate system
+     * 'left handed' coordinate system.
+     *
+     * Ignores scale.
+     */
+    static FORCEINLINE matrix_float4x4 ToARKitMatrix( const FTransform& InTransform, float WorldToMetersScale = 100.0f )
+    {
+		if(!ensure(WorldToMetersScale != 0.f))
+		{
+			WorldToMetersScale = 100.f;
+		}
+
+        matrix_float4x4 RetVal;
+
+        const FVector   Translation = InTransform.GetLocation() / WorldToMetersScale;
+        const FQuat     UnrealRotation = InTransform.GetRotation();
+        const FQuat     ARKitRotation = FQuat(UnrealRotation.Y, UnrealRotation.Z, -UnrealRotation.X, UnrealRotation.W);
+
+        const FMatrix   UnrealRotationMatrix = FRotationMatrix::Make(ARKitRotation);
+
+        RetVal.columns[0][0] = UnrealRotationMatrix.M[0][0]; RetVal.columns[0][1] = UnrealRotationMatrix.M[0][1]; RetVal.columns[0][2] = -UnrealRotationMatrix.M[0][2]; RetVal.columns[0][3] = UnrealRotationMatrix.M[0][3];
+        RetVal.columns[1][0] = UnrealRotationMatrix.M[1][0]; RetVal.columns[1][1] = UnrealRotationMatrix.M[1][1]; RetVal.columns[1][2] = UnrealRotationMatrix.M[1][2]; RetVal.columns[1][3] = UnrealRotationMatrix.M[1][3];
+        RetVal.columns[2][0] = -UnrealRotationMatrix.M[2][0]; RetVal.columns[2][1] = UnrealRotationMatrix.M[2][1]; RetVal.columns[2][2] = UnrealRotationMatrix.M[2][2]; RetVal.columns[2][3] = UnrealRotationMatrix.M[2][3];
+        RetVal.columns[3][0] = UnrealRotationMatrix.M[3][0]; RetVal.columns[3][1] = UnrealRotationMatrix.M[3][1]; RetVal.columns[3][2] = UnrealRotationMatrix.M[3][2]; RetVal.columns[3][3] = UnrealRotationMatrix.M[3][3];
+
+        // Set the translation element
+        RetVal.columns[3][2] = -Translation.X;
+        RetVal.columns[3][0] = Translation.Y;
+        RetVal.columns[3][1] = Translation.Z;
+
+        return RetVal;
+    }
+
+	/**
 	 * Convert's an ARKit 'Y up' 'right handed' coordinate system vector to Unreal's 'Z up' 
 	 * 'left handed' coordinate system.
 	 */
@@ -55,6 +88,25 @@ struct FAppleARKitConversion
 		return FVector( -RawYUpVector.z, RawYUpVector.x, RawYUpVector.y ) * ToUE4Scale();
 	}
 
+    /**
+     * Convert's an Unreal's 'Z up' to ARKit's 'Y up' vector
+     * 'left handed' coordinate system.
+     */
+    static FORCEINLINE vector_float3 ToARKitVector( const FVector& InFVector, float WorldToMetersScale = 100.0f )
+    {
+		if(!ensure(WorldToMetersScale != 0.f))
+		{
+			WorldToMetersScale = 100.f;
+		}
+
+        vector_float3 RetVal;
+        RetVal.x = InFVector.Y;
+        RetVal.y = InFVector.Z;
+        RetVal.z = -InFVector.X;
+
+        return RetVal / WorldToMetersScale;
+    }
+
 	static FORCEINLINE FGuid ToFGuid( uuid_t UUID )
 	{
 		FGuid AsGUID(
@@ -95,12 +147,12 @@ struct FAppleARKitConversion
 
 	#if SUPPORTS_ARKIT_1_5
 
-	static void InitImageDetection(UARSessionConfig* SessionConfig, ARWorldTrackingConfiguration* WorldConfig, TMap< FString, UARCandidateImage* >& CandidateImages, TMap< FString, CGImageRef >& ConvertedCandidateImages)
+	static NSSet* InitImageDetection(UARSessionConfig* SessionConfig, TMap< FString, UARCandidateImage* >& CandidateImages, TMap< FString, CGImageRef >& ConvertedCandidateImages)
 	{
 		const TArray<UARCandidateImage*>& ConfigCandidateImages = SessionConfig->GetCandidateImageList();
 		if (!ConfigCandidateImages.Num())
 		{
-			return;
+			return nullptr;
 		}
 
 		NSMutableSet* ConvertedImageSet = [[NSMutableSet new] autorelease];
@@ -136,12 +188,105 @@ struct FAppleARKitConversion
 				}
 			}
 		}
-		WorldConfig.detectionImages = ConvertedImageSet;
+		return ConvertedImageSet;
 	}
 
+	static void InitImageDetection(UARSessionConfig* SessionConfig, ARWorldTrackingConfiguration* WorldConfig, TMap< FString, UARCandidateImage* >& CandidateImages, TMap< FString, CGImageRef >& ConvertedCandidateImages)
+	{
+		if (FAppleARKitAvailability::SupportsARKit15())
+		{
+			WorldConfig.detectionImages = InitImageDetection(SessionConfig, CandidateImages, ConvertedCandidateImages);
+		}
+	//@joeg -- Added image tracking support
+	#if SUPPORTS_ARKIT_2_0
+		if (FAppleARKitAvailability::SupportsARKit20())
+		{
+			WorldConfig.maximumNumberOfTrackedImages = SessionConfig->GetMaxNumSimultaneousImagesTracked();
+		}
+	#endif
+	}
+	#endif
+
+	//@joeg -- Added image tracking support
+	#if SUPPORTS_ARKIT_2_0
+	static void InitImageDetection(UARSessionConfig* SessionConfig, ARImageTrackingConfiguration* ImageConfig, TMap< FString, UARCandidateImage* >& CandidateImages, TMap< FString, CGImageRef >& ConvertedCandidateImages)
+	{
+		ImageConfig.trackingImages = InitImageDetection(SessionConfig, CandidateImages, ConvertedCandidateImages);
+		ImageConfig.maximumNumberOfTrackedImages = SessionConfig->GetMaxNumSimultaneousImagesTracked();
+		ImageConfig.autoFocusEnabled = SessionConfig->ShouldEnableAutoFocus();
+	}
+
+	//@joeg -- Added environmental texture probe support
+	static AREnvironmentTexturing ToAREnvironmentTexturing(EAREnvironmentCaptureProbeType CaptureType)
+	{
+		switch (CaptureType)
+		{
+			case EAREnvironmentCaptureProbeType::Manual:
+			{
+				return AREnvironmentTexturingManual;
+			}
+			case EAREnvironmentCaptureProbeType::Automatic:
+			{
+				return AREnvironmentTexturingAutomatic;
+			}
+		}
+		return AREnvironmentTexturingNone;
+	}
+
+	static ARWorldMap* ToARWorldMap(const TArray<uint8>& WorldMapData)
+	{
+		NSData* WorldNSData = [NSData dataWithBytesNoCopy: (uint8*)WorldMapData.GetData() length: WorldMapData.Num() freeWhenDone: NO];
+		ARWorldMap* WorldMap = [NSKeyedUnarchiver unarchiveObjectWithData: WorldNSData];
+		if (WorldMap == nullptr)
+		{
+//@todo joeg -- Make a way to log from here
+//			UE_LOG(LogAppleARKit, Log, TEXT("Failed to load the world map data from the session object"));
+		}
+		return WorldMap;
+	}
+
+	static NSSet* ToARReferenceObjectSet(const TArray<UARCandidateObject*>& CandidateObjects, TMap< FString, UARCandidateObject* >& CandidateObjectMap)
+	{
+		CandidateObjectMap.Empty();
+
+		if (CandidateObjects.Num() == 0)
+		{
+			return nullptr;
+		}
+
+		NSMutableSet* ConvertedObjectSet = [[NSMutableSet new] autorelease];
+		for (UARCandidateObject* Candidate : CandidateObjects)
+		{
+			if (Candidate != nullptr && Candidate->GetCandidateObjectData().Num() > 0)
+			{
+				NSData* CandidateData = [NSData dataWithBytesNoCopy: (uint8*)Candidate->GetCandidateObjectData().GetData() length: Candidate->GetCandidateObjectData().Num() freeWhenDone: NO];
+				ARReferenceObject* RefObject = [NSKeyedUnarchiver unarchiveObjectWithData: CandidateData];
+				if (RefObject != nullptr)
+				{
+					// Store off so the session object can quickly match the anchor to our representation
+					// This stores it even if we weren't able to convert to apple's type for GC reasons
+					CandidateObjectMap.Add(Candidate->GetFriendlyName(), Candidate);
+					RefObject.name = Candidate->GetFriendlyName().GetNSString();
+					[ConvertedObjectSet addObject: RefObject];
+				}
+				else
+				{
+//@todo joeg -- Make a way to log from here
+//					UE_LOG(LogAppleARKit, Log, TEXT("Failed to convert to ARReferenceObject (%s)"), *Candidate->GetFriendlyName());
+				}
+			}
+			else
+			{
+//@todo joeg -- Make a way to log from here
+//				UE_LOG(LogAppleARKit, Log, TEXT("Missing candidate object data for ARCandidateObject (%s)"), Candidate != nullptr ? *Candidate->GetFriendlyName() : TEXT("null"));
+			}
+		}
+		return ConvertedObjectSet;
+	}
 	#endif
 
-	static ARConfiguration* ToARConfiguration( UARSessionConfig* SessionConfig, TMap< FString, UARCandidateImage* >& CandidateImages, TMap< FString, CGImageRef >& ConvertedCandidateImages )
+//@joeg -- Object detection addition
+	static ARConfiguration* ToARConfiguration( UARSessionConfig* SessionConfig, TMap< FString, UARCandidateImage* >& CandidateImages, TMap< FString, CGImageRef >& ConvertedCandidateImages, TMap< FString, UARCandidateObject* >& CandidateObjects )
 	{
 		EARSessionType SessionType = SessionConfig->GetSessionType();
 		ARConfiguration* SessionConfiguration = nullptr;
@@ -179,10 +324,70 @@ struct FAppleARKitConversion
 					// Add any images that wish to be detected
 					FAppleARKitConversion::InitImageDetection(SessionConfig, WorldTrackingConfiguration, CandidateImages, ConvertedCandidateImages);
 				}
+#endif
+//@joeg -- Added environmental texture probe support
+#if SUPPORTS_ARKIT_2_0
+				if (FAppleARKitAvailability::SupportsARKit20())
+				{
+					// Check for environment capture probe types
+					WorldTrackingConfiguration.environmentTexturing = ToAREnvironmentTexturing(SessionConfig->GetEnvironmentCaptureProbeType());
+					// Load the world if requested
+					if (SessionConfig->GetWorldMapData().Num() > 0)
+					{
+						ARWorldMap* WorldMap = ToARWorldMap(SessionConfig->GetWorldMapData());
+						WorldTrackingConfiguration.initialWorldMap = WorldMap;
+						[WorldMap release];
+					}
+					// Convert any candidate objects that are to be detected
+					WorldTrackingConfiguration.detectionObjects = ToARReferenceObjectSet(SessionConfig->GetCandidateObjectList(), CandidateObjects);
+				}
 #endif
 				SessionConfiguration = WorldTrackingConfiguration;
 				break;
 			}
+//@joeg -- Added image tracking support
+			case EARSessionType::Image:
+			{
+#if SUPPORTS_ARKIT_2_0
+				if (FAppleARKitAvailability::SupportsARKit20())
+				{
+					if (ARImageTrackingConfiguration.isSupported == FALSE)
+					{
+						return nullptr;
+					}
+					ARImageTrackingConfiguration* ImageTrackingConfiguration = [ARImageTrackingConfiguration new];
+					// Add any images that wish to be detected
+					InitImageDetection(SessionConfig, ImageTrackingConfiguration, CandidateImages, ConvertedCandidateImages);
+					SessionConfiguration = ImageTrackingConfiguration;
+				}
+#endif
+				break;
+			}
+//@joeg -- Added object scanning support
+			case EARSessionType::ObjectScanning:
+			{
+#if SUPPORTS_ARKIT_2_0
+				if (FAppleARKitAvailability::SupportsARKit20())
+				{
+					if (ARObjectScanningConfiguration.isSupported == FALSE)
+					{
+						return nullptr;
+					}
+					ARObjectScanningConfiguration* ObjectScanningConfiguration = [ARObjectScanningConfiguration new];
+					if (EnumHasAnyFlags(EARPlaneDetectionMode::HorizontalPlaneDetection, SessionConfig->GetPlaneDetectionMode()))
+					{
+						ObjectScanningConfiguration.planeDetection |= ARPlaneDetectionHorizontal;
+					}
+					if (EnumHasAnyFlags(EARPlaneDetectionMode::VerticalPlaneDetection, SessionConfig->GetPlaneDetectionMode()))
+					{
+						ObjectScanningConfiguration.planeDetection |= ARPlaneDetectionVertical;
+					}
+					ObjectScanningConfiguration.autoFocusEnabled = SessionConfig->ShouldEnableAutoFocus();
+					SessionConfiguration = ObjectScanningConfiguration;
+				}
+#endif
+				break;
+			}
 			default:
 				return nullptr;
 		}
@@ -206,6 +411,10 @@ enum class EAppleAnchorType : uint8
 	PlaneAnchor,
 	FaceAnchor,
 	ImageAnchor,
+//@joeg -- Added environmental texture probe support
+	EnvironmentProbeAnchor,
+//@joeg -- Object detection
+	ObjectAnchor,
 	MAX
 };
 
@@ -227,20 +436,35 @@ struct FAppleARKitAnchorData
 	{
 	}
 
-	FAppleARKitAnchorData(FGuid InAnchorGuid, FTransform InTransform, FString InDetectedImageName)
-		: Transform( InTransform )
-		, AnchorType( EAppleAnchorType::ImageAnchor )
-		, AnchorGUID( InAnchorGuid )
-		, DetectedImageName( MoveTemp(InDetectedImageName) )
+//@joeg -- Eye tracking support
+	FAppleARKitAnchorData(FGuid InAnchorGuid, FTransform InTransform, FARBlendShapeMap InBlendShapes, TArray<FVector> InFaceVerts, FTransform InLeftEyeTransform, FTransform InRightEyeTransform, FVector InLookAtTarget)
+	: Transform( InTransform )
+	, AnchorType( EAppleAnchorType::FaceAnchor )
+	, AnchorGUID( InAnchorGuid )
+	, BlendShapes( MoveTemp(InBlendShapes) )
+	, FaceVerts( MoveTemp(InFaceVerts) )
+	, LeftEyeTransform( InLeftEyeTransform )
+	, RightEyeTransform( InRightEyeTransform )
+	, LookAtTarget( InLookAtTarget )
 	{
 	}
 
-	FAppleARKitAnchorData(FGuid InAnchorGuid, FTransform InTransform, FARBlendShapeMap InBlendShapes, TArray<FVector> InFaceVerts)
-		: Transform( InTransform )
-		, AnchorType( EAppleAnchorType::FaceAnchor )
-		, AnchorGUID( InAnchorGuid )
-		, BlendShapes( MoveTemp(InBlendShapes) )
-		, FaceVerts( MoveTemp(InFaceVerts) )
+//@joeg -- Changed so object and image detection share
+	FAppleARKitAnchorData(FGuid InAnchorGuid, FTransform InTransform, EAppleAnchorType InAnchorType, FString InDetectedAnchorName)
+    : Transform( InTransform )
+    , AnchorType( InAnchorType )
+    , AnchorGUID( InAnchorGuid )
+    , DetectedAnchorName( MoveTemp(InDetectedAnchorName) )
+    {
+    }
+
+//@joeg -- Added environmental texture probe support
+	FAppleARKitAnchorData(FGuid InAnchorGuid, FTransform InTransform, FVector InExtent, id<MTLTexture> InProbeTexture)
+	: Transform( InTransform )
+	, AnchorType( EAppleAnchorType::EnvironmentProbeAnchor )
+	, AnchorGUID( InAnchorGuid )
+	, Extent(InExtent)
+	, ProbeTexture(InProbeTexture)
 	{
 	}
 
@@ -251,12 +475,21 @@ struct FAppleARKitAnchorData
 	FVector Extent;
 	TArray<FVector> BoundaryVerts;
 
-	FString DetectedImageName;
-
 	FARBlendShapeMap BlendShapes;
 	TArray<FVector> FaceVerts;
 	// Note: the index buffer never changes so can be safely read once
 	static TArray<int32> FaceIndices;
+
+//@joeg -- Changed this so both Image and Object anchors can share
+	FString DetectedAnchorName;
+
+//@joeg -- Environment texturing support
+	id<MTLTexture> ProbeTexture;
+
+//@joeg -- Eye tracking support
+	FTransform LeftEyeTransform;
+	FTransform RightEyeTransform;
+	FVector LookAtTarget;
 };
 
 namespace ARKitUtil
diff --git a/Engine/Plugins/Runtime/AppleARKit/Source/AppleARKit/Public/AppleARKitFrame.h b/Engine/Plugins/Runtime/AppleARKit/Source/AppleARKit/Public/AppleARKitFrame.h
--- a/Engine/Plugins/Runtime/AppleARKit/Source/AppleARKit/Public/AppleARKitFrame.h
+++ b/Engine/Plugins/Runtime/AppleARKit/Source/AppleARKit/Public/AppleARKitFrame.h
@@ -78,6 +78,10 @@ struct APPLEARKIT_API FAppleARKitFrame
 	 * A light estimate representing the estimated light in the scene.
 	 */
 	FAppleARKitLightEstimate LightEstimate;
+	
+//@joeg -- ARKit 2.0 support
+	/** The current world mapping state is reported on the frame */
+	EARWorldMappingState WorldMappingState;
 
 	/* 
 	 * When adding new member variables, don't forget to handle them in the copy constructor and
diff --git a/Engine/Plugins/Runtime/AppleARKit/Source/AppleARKit/Public/AppleARKitTextures.h b/Engine/Plugins/Runtime/AppleARKit/Source/AppleARKit/Public/AppleARKitTextures.h
--- a/Engine/Plugins/Runtime/AppleARKit/Source/AppleARKit/Public/AppleARKitTextures.h
+++ b/Engine/Plugins/Runtime/AppleARKit/Source/AppleARKit/Public/AppleARKitTextures.h
@@ -84,4 +84,39 @@ private:
 #endif
 };
 
+//@joeg -- Added for environment capture support
+UCLASS(BlueprintType)
+class APPLEARKIT_API UAppleARKitEnvironmentCaptureProbeTexture :
+	public UAREnvironmentCaptureProbeTexture,
+	public IAppleImageInterface
+{
+	GENERATED_UCLASS_BODY()
+	
+public:
+	
+	// UTexture interface implementation
+	virtual void BeginDestroy() override;
+	virtual FTextureResource* CreateResource() override;
+	virtual EMaterialValueType GetMaterialType() const override { return MCT_TextureExternal; }
+	virtual float GetSurfaceWidth() const override { return Size.X; }
+	virtual float GetSurfaceHeight() const override { return Size.Y; }
+	virtual FGuid GetExternalTextureGuid() const override { return ExternalTextureGuid; }
+	// End UTexture interface
+	
+	virtual EAppleTextureType GetTextureType() const override { return EAppleTextureType::MetalTexture; }
+
+#if PLATFORM_MAC || PLATFORM_IOS
+	/** Sets any initialization data */
+	virtual void Init(float InTimestamp, id<MTLTexture> InEnvironmentTexture);
+
+	// IAppleImageInterface interface implementation
+	virtual id<MTLTexture> GetMetalTexture() const override { return MetalTexture; }
+	// End IAppleImageInterface interface
+
+private:
+	/** The Apple specific representation of the ar environment texture */
+	id<MTLTexture> MetalTexture;
+#endif
+};
+
 
diff --git a/Engine/Plugins/Runtime/AppleARKit/Source/AppleARKit/Public/AppleARKitTrackable.h b/Engine/Plugins/Runtime/AppleARKit/Source/AppleARKit/Public/AppleARKitTrackable.h
new file mode 100644
--- /dev/null
+++ b/Engine/Plugins/Runtime/AppleARKit/Source/AppleARKit/Public/AppleARKitTrackable.h
@@ -0,0 +1,30 @@
+// Copyright 1998-2018 Epic Games, Inc. All Rights Reserved.
+
+#pragma once
+
+#include "AppleARKitAvailability.h"
+#include "ARTrackable.h"
+#include "AppleARKitTrackable.generated.h"
+
+class UAppleARKitEnvironmentCaptureProbeTexture;
+
+//@joeg -- Added support for environment capture
+UCLASS(BlueprintType)
+class APPLEARKIT_API UAppleARKitEnvironmentCaptureProbe :
+	public UAREnvironmentCaptureProbe
+{
+	GENERATED_BODY()
+
+public:
+	UAppleARKitEnvironmentCaptureProbe();
+	
+#if PLATFORM_MAC || PLATFORM_IOS
+	/** Updates the current capture with the new metal texture. This will trigger a render resource update if the textures have changed */
+	void UpdateEnvironmentCapture(const TSharedRef<FARSystemBase, ESPMode::ThreadSafe>& InTrackingSystem, uint32 FrameNumber, double Timestamp, const FTransform& InLocalToTrackingTransform, const FTransform& InAlignmentTransform, FVector InExtent, id<MTLTexture> InMetalTexture);
+#endif
+
+private:
+	/** The cube map of the reflected environment */
+	UPROPERTY()
+	UAppleARKitEnvironmentCaptureProbeTexture* ARKitEnvironmentTexture;
+};
diff --git a/Engine/Plugins/Runtime/AppleARKitFaceSupport/Source/AppleARKitFaceSupport/Private/AppleARKitFaceMeshConversion.h b/Engine/Plugins/Runtime/AppleARKitFaceSupport/Source/AppleARKitFaceSupport/Private/AppleARKitFaceMeshConversion.h
--- a/Engine/Plugins/Runtime/AppleARKitFaceSupport/Source/AppleARKitFaceSupport/Private/AppleARKitFaceMeshConversion.h
+++ b/Engine/Plugins/Runtime/AppleARKitFaceSupport/Source/AppleARKitFaceSupport/Private/AppleARKitFaceMeshConversion.h
@@ -2,9 +2,11 @@
 
 #pragma once
 
+#include "AppleARKitSettings.h"
+
 #if SUPPORTS_ARKIT_1_0
 
-FORCEINLINE TArray<int32> To32BitIndexBuffer(const int16_t* Indices, const uint64 IndexCount)
+static FORCEINLINE TArray<int32> To32BitIndexBuffer(const int16_t* Indices, const uint64 IndexCount)
 {
 	check(IndexCount % 3 == 0);
 
@@ -21,7 +23,7 @@ FORCEINLINE TArray<int32> To32BitIndexBuffer(const int16_t* Indices, const uint6
 }
 
 // @todo JoeG - An option for which way to orient tris (down +X or -X)
-FORCEINLINE TArray<FVector> ToVertexBuffer(const vector_float3* Vertices, const uint64 VertexCount)
+static FORCEINLINE TArray<FVector> ToVertexBuffer(const vector_float3* Vertices, const uint64 VertexCount)
 {
 	TArray<FVector> VertexBuffer;
 	VertexBuffer.AddUninitialized(VertexCount);
@@ -33,7 +35,7 @@ FORCEINLINE TArray<FVector> ToVertexBuffer(const vector_float3* Vertices, const
 	return VertexBuffer;
 }
 
-FARBlendShapeMap ToBlendShapeMap(NSDictionary<ARBlendShapeLocation,NSNumber *>* BlendShapes, const FTransform& Transform)
+static FORCEINLINE FARBlendShapeMap ToBlendShapeMap(NSDictionary<ARBlendShapeLocation,NSNumber *>* BlendShapes, const FTransform& Transform, const FTransform& LeftEyeTransform, const FTransform& RightEyeTransform)
 {
 	FARBlendShapeMap BlendShapeMap;
 	FRotator TrackedRot(Transform.GetRotation());
@@ -44,6 +46,22 @@ FARBlendShapeMap ToBlendShapeMap(NSDictionary<ARBlendShapeLocation,NSNumber *>*
 	BlendShapeMap.Add(EARFaceBlendShape::HeadYaw, HeadYaw);
 	BlendShapeMap.Add(EARFaceBlendShape::HeadPitch, HeadPitch);
 	BlendShapeMap.Add(EARFaceBlendShape::HeadRoll, HeadRoll);
+//@joeg -- Eye tracking support
+	FRotator LeftEyeRot(LeftEyeTransform.GetRotation());
+	float LeftEyeYaw = (float)LeftEyeRot.Yaw / 180.f;
+	float LeftEyePitch = (float)LeftEyeRot.Pitch / 180.f;
+	float LeftEyeRoll = (float)LeftEyeRot.Roll / 180.f;
+	BlendShapeMap.Add(EARFaceBlendShape::LeftEyeYaw, LeftEyeYaw);
+	BlendShapeMap.Add(EARFaceBlendShape::LeftEyePitch, LeftEyePitch);
+	BlendShapeMap.Add(EARFaceBlendShape::LeftEyeRoll, LeftEyeRoll);
+	FRotator RightEyeRot(RightEyeTransform.GetRotation());
+	float RightEyeYaw = (float)RightEyeRot.Yaw / 180.f;
+	float RightEyePitch = (float)RightEyeRot.Pitch / 180.f;
+	float RightEyeRoll = (float)RightEyeRot.Roll / 180.f;
+	BlendShapeMap.Add(EARFaceBlendShape::RightEyeYaw, RightEyeYaw);
+	BlendShapeMap.Add(EARFaceBlendShape::RightEyePitch, RightEyePitch);
+	BlendShapeMap.Add(EARFaceBlendShape::RightEyeRoll, RightEyeRoll);
+
 
 #define SET_BLEND_SHAPE(AppleShape, UE4Shape) \
 	if (BlendShapes[AppleShape]) \
@@ -156,7 +174,13 @@ FARBlendShapeMap ToBlendShapeMap(NSDictionary<ARBlendShapeLocation,NSNumber *>*
 	SET_BLEND_SHAPE(ARBlendShapeLocationMouthShrugUpper, EARFaceBlendShape::MouthShrugUpper);
 	SET_BLEND_SHAPE(ARBlendShapeLocationBrowInnerUp, EARFaceBlendShape::BrowInnerUp);
 	SET_BLEND_SHAPE(ARBlendShapeLocationCheekPuff, EARFaceBlendShape::CheekPuff);
-
+//@joeg -- Tongue blend shape
+#if SUPPORTS_ARKIT_2_0
+	if (FAppleARKitAvailability::SupportsARKit20())
+	{
+		SET_BLEND_SHAPE(ARBlendShapeLocationTongueOut, EARFaceBlendShape::TongueOut);
+	}
+#endif
 
 #undef SET_BLEND_SHAPE
 
diff --git a/Engine/Plugins/Runtime/AppleARKitFaceSupport/Source/AppleARKitFaceSupport/Private/AppleARKitFaceSupportImpl.cpp b/Engine/Plugins/Runtime/AppleARKitFaceSupport/Source/AppleARKitFaceSupport/Private/AppleARKitFaceSupportImpl.cpp
--- a/Engine/Plugins/Runtime/AppleARKitFaceSupport/Source/AppleARKitFaceSupport/Private/AppleARKitFaceSupportImpl.cpp
+++ b/Engine/Plugins/Runtime/AppleARKitFaceSupport/Source/AppleARKitFaceSupport/Private/AppleARKitFaceSupportImpl.cpp
@@ -18,12 +18,27 @@ static TSharedPtr<FAppleARKitAnchorData> MakeAnchorData(ARAnchor* Anchor)
     if ([Anchor isKindOfClass:[ARFaceAnchor class]])
     {
         ARFaceAnchor* FaceAnchor = (ARFaceAnchor*)Anchor;
-        NewAnchor = MakeShared<FAppleARKitAnchorData>(
+//@joeg -- Eye tracking support
+		FTransform LeftEyeTransform;
+		FTransform RightEyeTransform;
+		FVector LookAtTarget;
+#if SUPPORTS_ARKIT_2_0
+		if (FAppleARKitAvailability::SupportsARKit20())
+		{
+			LeftEyeTransform = FAppleARKitConversion::ToFTransform(FaceAnchor.leftEyeTransform);
+			RightEyeTransform = FAppleARKitConversion::ToFTransform(FaceAnchor.rightEyeTransform);
+			LookAtTarget = FAppleARKitConversion::ToFVector(FaceAnchor.lookAtPoint);
+		}
+#endif
+		NewAnchor = MakeShared<FAppleARKitAnchorData>(
 			FAppleARKitConversion::ToFGuid(FaceAnchor.identifier),
 			FAppleARKitConversion::ToFTransform(FaceAnchor.transform),
-			ToBlendShapeMap(FaceAnchor.blendShapes, FAppleARKitConversion::ToFTransform(FaceAnchor.transform)),
-			ToVertexBuffer(FaceAnchor.geometry.vertices, FaceAnchor.geometry.vertexCount)
-        );
+			ToBlendShapeMap(FaceAnchor.blendShapes, FAppleARKitConversion::ToFTransform(FaceAnchor.transform), LeftEyeTransform, RightEyeTransform),
+			ToVertexBuffer(FaceAnchor.geometry.vertices, FaceAnchor.geometry.vertexCount),
+			LeftEyeTransform,
+			RightEyeTransform,
+			LookAtTarget
+		);
         // Only convert from 16bit to 32bit once
         if (FAppleARKitAnchorData::FaceIndices.Num() == 0)
         {
diff --git a/Engine/Plugins/Runtime/GoogleARCore/Source/GoogleARCoreBase/Private/GoogleARCoreXRTrackingSystem.h b/Engine/Plugins/Runtime/GoogleARCore/Source/GoogleARCoreBase/Private/GoogleARCoreXRTrackingSystem.h
--- a/Engine/Plugins/Runtime/GoogleARCore/Source/GoogleARCoreBase/Private/GoogleARCoreXRTrackingSystem.h
+++ b/Engine/Plugins/Runtime/GoogleARCore/Source/GoogleARCoreBase/Private/GoogleARCoreXRTrackingSystem.h
@@ -60,6 +60,11 @@ protected:
 	virtual void OnRemovePin(UARPin* PinToRemove) override;
 	virtual UARTextureCameraImage* OnGetCameraImage() override { return nullptr; }
 	virtual UARTextureCameraDepth* OnGetCameraDepth() override { return nullptr; }
+	virtual bool OnAddManualEnvironmentCaptureProbe(FVector Location, FVector Extent) { return false; }
+	virtual TSharedPtr<FARGetCandidateObjectAsyncTask, ESPMode::ThreadSafe> OnGetCandidateObject(FVector Location, FVector Extent) const { return TSharedPtr<FARGetCandidateObjectAsyncTask, ESPMode::ThreadSafe>(); }
+	virtual TSharedPtr<FARSaveWorldAsyncTask, ESPMode::ThreadSafe> OnSaveWorld() const { return TSharedPtr<FARSaveWorldAsyncTask, ESPMode::ThreadSafe>(); }
+// @todo -- support this properly
+	virtual EARWorldMappingState OnGetWorldMappingStatus() const { return EARWorldMappingState::StillMappingNotRelocalizable; }
 	//~IARSystemSupport
 
 private:
diff --git a/Engine/Source/Runtime/AugmentedReality/Private/ARBlueprintLibrary.cpp b/Engine/Source/Runtime/AugmentedReality/Private/ARBlueprintLibrary.cpp
--- a/Engine/Source/Runtime/AugmentedReality/Private/ARBlueprintLibrary.cpp
+++ b/Engine/Source/Runtime/AugmentedReality/Private/ARBlueprintLibrary.cpp
@@ -339,3 +339,122 @@ UARTextureCameraDepth* UARBlueprintLibrary::GetCameraDepth()
 	}
 	return Depth;
 }
+
+//@joeg -- Helpers to make things easier
+
+TArray<UARPlaneGeometry*> UARBlueprintLibrary::GetAllTrackedPlanes()
+{
+	TArray<UARPlaneGeometry*> ResultList;
+	
+	auto ARSystem = GetARSystem();
+	if (ensure(ARSystem.IsValid()))
+	{
+		for (UARTrackedGeometry* Geo : ARSystem->GetAllTrackedGeometries())
+		{
+			if (UARPlaneGeometry* Item = Cast<UARPlaneGeometry>(Geo))
+			{
+				ResultList.Add(Item);
+			}
+		}
+	}
+	return ResultList;
+}
+
+TArray<UARTrackedPoint*> UARBlueprintLibrary::GetAllTrackedPoints()
+{
+	TArray<UARTrackedPoint*> ResultList;
+	
+	auto ARSystem = GetARSystem();
+	if (ensure(ARSystem.IsValid()))
+	{
+		for (UARTrackedGeometry* Geo : ARSystem->GetAllTrackedGeometries())
+		{
+			if (UARTrackedPoint* Item = Cast<UARTrackedPoint>(Geo))
+			{
+				ResultList.Add(Item);
+			}
+		}
+	}
+	return ResultList;
+}
+
+TArray<UARTrackedImage*> UARBlueprintLibrary::GetAllTrackedImages()
+{
+	TArray<UARTrackedImage*> ResultList;
+	
+	auto ARSystem = GetARSystem();
+	if (ensure(ARSystem.IsValid()))
+	{
+		for (UARTrackedGeometry* Geo : ARSystem->GetAllTrackedGeometries())
+		{
+			if (UARTrackedImage* Item = Cast<UARTrackedImage>(Geo))
+			{
+				ResultList.Add(Item);
+			}
+		}
+	}
+	return ResultList;
+}
+
+TArray<UAREnvironmentCaptureProbe*> UARBlueprintLibrary::GetAllTrackedEnvironmentCaptureProbes()
+{
+	TArray<UAREnvironmentCaptureProbe*> ResultList;
+	
+	auto ARSystem = GetARSystem();
+	if (ensure(ARSystem.IsValid()))
+	{
+		for (UARTrackedGeometry* Geo : ARSystem->GetAllTrackedGeometries())
+		{
+			if (UAREnvironmentCaptureProbe* Item = Cast<UAREnvironmentCaptureProbe>(Geo))
+			{
+				ResultList.Add(Item);
+			}
+		}
+	}
+	return ResultList;
+}
+
+//@joeg -- Added environmental texture probe support
+bool UARBlueprintLibrary::AddManualEnvironmentCaptureProbe(FVector Location, FVector Extent)
+{
+	auto ARSystem = GetARSystem();
+	if (ensure(ARSystem.IsValid()))
+	{
+		return ARSystem->AddManualEnvironmentCaptureProbe(Location, Extent);
+	}
+	return false;
+}
+
+EARWorldMappingState UARBlueprintLibrary::GetWorldMappingStatus()
+{
+	auto ARSystem = GetARSystem();
+	if (ensure(ARSystem.IsValid()))
+	{
+		return ARSystem->GetWorldMappingStatus();
+	}
+	return EARWorldMappingState::NotAvailable;
+}
+
+TSharedPtr<FARSaveWorldAsyncTask, ESPMode::ThreadSafe> UARBlueprintLibrary::SaveWorld()
+{
+	auto ARSystem = GetARSystem();
+	if (ensure(ARSystem.IsValid()))
+	{
+		return ARSystem->SaveWorld();
+		
+	}
+	return TSharedPtr<FARSaveWorldAsyncTask, ESPMode::ThreadSafe>();
+}
+
+TSharedPtr<FARGetCandidateObjectAsyncTask, ESPMode::ThreadSafe> UARBlueprintLibrary::GetCandidateObject(FVector Location, FVector Extent)
+{
+	auto ARSystem = GetARSystem();
+	if (ensure(ARSystem.IsValid()))
+	{
+		return ARSystem->GetCandidateObject(Location, Extent);
+		
+	}
+	return TSharedPtr<FARGetCandidateObjectAsyncTask, ESPMode::ThreadSafe>();
+}
+
+//@joeg -- End additions
diff --git a/Engine/Source/Runtime/AugmentedReality/Private/ARSessionConfig.cpp b/Engine/Source/Runtime/AugmentedReality/Private/ARSessionConfig.cpp
--- a/Engine/Source/Runtime/AugmentedReality/Private/ARSessionConfig.cpp
+++ b/Engine/Source/Runtime/AugmentedReality/Private/ARSessionConfig.cpp
@@ -14,6 +14,8 @@ UARSessionConfig::UARSessionConfig()
 , FrameSyncMode(EARFrameSyncMode::SyncTickWithoutCameraImage)
 , bEnableAutomaticCameraOverlay(true)
 , bEnableAutomaticCameraTracking(true)
+//@joeg -- Added image tracking support
+, MaxNumSimultaneousImagesTracked(1)
 {
 }
 
@@ -64,6 +66,46 @@ const TArray<UARCandidateImage*>& UARSessionConfig::GetCandidateImageList() cons
 	return CandidateImages;
 }
 
+//@joeg -- Added image tracking support
+int32 UARSessionConfig::GetMaxNumSimultaneousImagesTracked() const
+{
+    return MaxNumSimultaneousImagesTracked;
+}
+
+//@joeg -- Added environmental texture probe support
+/** @see EnvironmentCaptureProbeType */
+EAREnvironmentCaptureProbeType UARSessionConfig::GetEnvironmentCaptureProbeType() const
+{
+	return EnvironmentCaptureProbeType;
+}
+
+//@joeg -- Load/Save world support
+const TArray<uint8>& UARSessionConfig::GetWorldMapData() const
+{
+	return WorldMapData;
+}
+
+void UARSessionConfig::SetWorldMapData(TArray<uint8> InWorldMapData)
+{
+	WorldMapData = MoveTemp(InWorldMapData);
+}
+//@joeg -- end load/save
+
+//@joeg -- Object detection
+const TArray<UARCandidateObject*>& UARSessionConfig::GetCandidateObjectList() const
+{
+	return CandidateObjects;
+}
+
+void UARSessionConfig::AddCandidateObject(UARCandidateObject* CandidateObject)
+{
+	if (CandidateObject != nullptr)
+	{
+		CandidateObjects.Add(CandidateObject);
+	}
+}
+//@joeg -- end object detection
+
 void UARSessionConfig::Serialize(FArchive& Ar)
 {
 	Ar.UsingCustomVersion(FVRObjectVersion::GUID);
diff --git a/Engine/Source/Runtime/AugmentedReality/Private/ARSystem.cpp b/Engine/Source/Runtime/AugmentedReality/Private/ARSystem.cpp
--- a/Engine/Source/Runtime/AugmentedReality/Private/ARSystem.cpp
+++ b/Engine/Source/Runtime/AugmentedReality/Private/ARSystem.cpp
@@ -4,6 +4,7 @@
 #include "IXRTrackingSystem.h"
 #include "Features/IModularFeatures.h"
 #include "ARBlueprintLibrary.h"
+#include "ARBlueprintProxy.h"
 #include "ARSessionConfig.h"
 #include "GeneralProjectSettings.h"
 #include "Engine/Engine.h"
@@ -23,7 +24,8 @@ void FARSystemBase::InitializeARSystem()
 	IModularFeatures::Get().RegisterModularFeature(FARSystemBase::GetModularFeatureName(), this);
 	
 	UARBlueprintLibrary::RegisterAsARSystem( SharedThis(this) );
-	
+	UARBaseAsyncTaskBlueprintProxy::RegisterAsARSystem( SharedThis(this) );
+
 	OnARSystemInitialized();
 }
 
@@ -32,6 +34,7 @@ FARSystemBase::~FARSystemBase()
 	IModularFeatures::Get().UnregisterModularFeature(FARSystemBase::GetModularFeatureName(), this);
 	
 	UARBlueprintLibrary::RegisterAsARSystem( nullptr );
+	UARBaseAsyncTaskBlueprintProxy::RegisterAsARSystem( nullptr );
 }
 
 EARTrackingQuality FARSystemBase::GetTrackingQuality() const
@@ -108,6 +111,29 @@ UARTextureCameraDepth* FARSystemBase::GetCameraDepth()
 	return OnGetCameraDepth();
 }
 
+//@joeg -- ARKit 2.0 additions
+
+bool FARSystemBase::AddManualEnvironmentCaptureProbe(FVector Location, FVector Extent)
+{
+	return OnAddManualEnvironmentCaptureProbe(Location, Extent);
+}
+
+TSharedPtr<FARGetCandidateObjectAsyncTask, ESPMode::ThreadSafe> FARSystemBase::GetCandidateObject(FVector Location, FVector Extent) const
+{
+	return OnGetCandidateObject(Location, Extent);
+}
+
+TSharedPtr<FARSaveWorldAsyncTask, ESPMode::ThreadSafe> FARSystemBase::SaveWorld() const
+{
+	return OnSaveWorld();
+}
+
+EARWorldMappingState FARSystemBase::GetWorldMappingStatus() const
+{
+	return OnGetWorldMappingStatus();
+}
+//@joeg -- End additions
+
 bool FARSystemBase::IsSessionTypeSupported(EARSessionType SessionType) const
 {
 	return OnIsTrackingTypeSupported(SessionType);
diff --git a/Engine/Source/Runtime/AugmentedReality/Private/ARTextures.cpp b/Engine/Source/Runtime/AugmentedReality/Private/ARTextures.cpp
--- a/Engine/Source/Runtime/AugmentedReality/Private/ARTextures.cpp
+++ b/Engine/Source/Runtime/AugmentedReality/Private/ARTextures.cpp
@@ -16,3 +16,10 @@ UARTextureCameraDepth::UARTextureCameraDepth(const FObjectInitializer& ObjectIni
 	: Super(ObjectInitializer)
 {
 }
+
+//@joeg -- Added environmental texture probe support
+UAREnvironmentCaptureProbeTexture::UAREnvironmentCaptureProbeTexture(const FObjectInitializer& ObjectInitializer)
+	: Super(ObjectInitializer)
+	, TextureType(EARTextureType::EnvironmentCapture)
+{
+}
diff --git a/Engine/Source/Runtime/AugmentedReality/Private/ARTrackable.cpp b/Engine/Source/Runtime/AugmentedReality/Private/ARTrackable.cpp
--- a/Engine/Source/Runtime/AugmentedReality/Private/ARTrackable.cpp
+++ b/Engine/Source/Runtime/AugmentedReality/Private/ARTrackable.cpp
@@ -173,7 +173,7 @@ void UARTrackedImage::UpdateTrackedGeometry(const TSharedRef<FARSystemBase, ESPM
 	DetectedImage = InDetectedImage;
 }
 
-void UARFaceGeometry::UpdateTrackedGeometry(const TSharedRef<FARSystemBase, ESPMode::ThreadSafe>& InTrackingSystem, uint32 FrameNumber, double Timestamp, const FTransform& InLocalToTrackingTransform, const FTransform& InAlignmentTransform, FARBlendShapeMap& InBlendShapes, TArray<FVector>& InVertices, const TArray<int32>& Indices)
+void UARFaceGeometry::UpdateFaceGeometry(const TSharedRef<FARSystemBase, ESPMode::ThreadSafe>& InTrackingSystem, uint32 FrameNumber, double Timestamp, const FTransform& InLocalToTrackingTransform, const FTransform& InAlignmentTransform, FARBlendShapeMap& InBlendShapes, TArray<FVector>& InVertices, const TArray<int32>& Indices, const FTransform& InLeftEyeTransform, const FTransform& InRightEyeTransform, const FVector& InLookAtTarget)
 {
 	Super::UpdateTrackedGeometry(InTrackingSystem, FrameNumber, Timestamp, InLocalToTrackingTransform, InAlignmentTransform);
 	BlendShapes = MoveTemp(InBlendShapes);
@@ -183,6 +183,11 @@ void UARFaceGeometry::UpdateTrackedGeometry(const TSharedRef<FARSystemBase, ESPM
 	{
 		IndexBuffer = Indices;
 	}
+	
+//@joeg -- Eye tracking support
+	LeftEyeTransform = InLeftEyeTransform;
+	RightEyeTransform = InRightEyeTransform;
+	LookAtTarget = InLookAtTarget;
 }
 
 void UARTrackedPoint::DebugDraw(UWorld* World, const FLinearColor& OutlineColor, float OutlineThickness, float PersistForSeconds /*= 0.0f*/) const
@@ -219,3 +224,77 @@ const TMap<EARFaceBlendShape, float> UARFaceGeometry::GetBlendShapes() const
 {
 	return BlendShapes;
 }
+
+//@joeg -- Eye tracking support
+const FTransform& UARFaceGeometry::GetLocalSpaceEyeTransform(EAREye Eye) const
+{
+	if (Eye == EAREye::LeftEye)
+	{
+		return LeftEyeTransform;
+	}
+	return RightEyeTransform;
+}
+
+FTransform UARFaceGeometry::GetWorldSpaceEyeTransform(EAREye Eye) const
+{
+	if (Eye == EAREye::LeftEye)
+	{
+		return GetLocalToWorldTransform() * LeftEyeTransform;
+	}
+	return GetLocalToWorldTransform() * RightEyeTransform;
+}
+//@joeg -- end eye tracking
+
+
+//@joeg -- Added environmental texture probe support
+UAREnvironmentCaptureProbe::UAREnvironmentCaptureProbe()
+	: Super()
+	, EnvironmentCaptureTexture(nullptr)
+{
+}
+
+void UAREnvironmentCaptureProbe::DebugDraw(UWorld* World, const FLinearColor& OutlineColor, float OutlineThickness, float PersistForSeconds) const
+{
+	const FTransform LocalToWorldTransform = GetLocalToWorldTransform();
+	const FString CurAnchorDebugName = GetDebugName().ToString();
+	const FColor OutlineRGB =OutlineColor.ToFColor(false);
+
+	ARDebugHelpers::DrawDebugString( World, LocalToWorldTransform.GetLocation(), CurAnchorDebugName, 0.25f * OutlineThickness, OutlineRGB, PersistForSeconds, true );
+	
+	DrawDebugBox( World, LocalToWorldTransform.GetLocation(), Extent, LocalToWorldTransform.GetRotation(), OutlineRGB, false, PersistForSeconds, 0, 0.1f * OutlineThickness );
+}
+
+void UAREnvironmentCaptureProbe::UpdateEnvironmentCapture(const TSharedRef<FARSystemBase, ESPMode::ThreadSafe>& InTrackingSystem, uint32 InFrameNumber, double InTimestamp, const FTransform& InLocalToTrackingTransform, const FTransform& InAlignmentTransform, FVector InExtent)
+{
+	Super::UpdateTrackedGeometry(InTrackingSystem, InFrameNumber, InTimestamp, InLocalToTrackingTransform, InAlignmentTransform);
+
+	Extent = InExtent;
+}
+
+FVector UAREnvironmentCaptureProbe::GetExtent() const
+{
+	return Extent;
+}
+
+UAREnvironmentCaptureProbeTexture* UAREnvironmentCaptureProbe::GetEnvironmentCaptureTexture() const
+{
+	return EnvironmentCaptureTexture;
+}
+
+//@joeg -- Object detection
+void UARTrackedObject::DebugDraw(UWorld* World, const FLinearColor& OutlineColor, float OutlineThickness, float PersistForSeconds /*= 0.0f*/) const
+{
+	const FTransform LocalToWorldTransform = GetLocalToWorldTransform();
+	const FString CurAnchorDebugName = GetDebugName().ToString();
+	const FColor OutlineRGB =OutlineColor.ToFColor(false);
+	ARDebugHelpers::DrawDebugString( World, LocalToWorldTransform.GetLocation(), CurAnchorDebugName, 0.25f*OutlineThickness, OutlineRGB, PersistForSeconds, true);
+	
+	DrawDebugPoint(World, LocalToWorldTransform.GetLocation(), 0.5f, OutlineRGB, false, PersistForSeconds, 0);
+}
+
+void UARTrackedObject::UpdateTrackedGeometry(const TSharedRef<FARSystemBase, ESPMode::ThreadSafe>& InTrackingSystem, uint32 FrameNumber, double Timestamp, const FTransform& InLocalToTrackingTransform, const FTransform& InAlignmentTransform, UARCandidateObject* InDetectedObject)
+{
+	Super::UpdateTrackedGeometry(InTrackingSystem, FrameNumber, Timestamp, InLocalToTrackingTransform, InAlignmentTransform);
+	DetectedObject = InDetectedObject;
+}
+
diff --git a/Engine/Source/Runtime/AugmentedReality/Private/ARTypes.cpp b/Engine/Source/Runtime/AugmentedReality/Private/ARTypes.cpp
--- a/Engine/Source/Runtime/AugmentedReality/Private/ARTypes.cpp
+++ b/Engine/Source/Runtime/AugmentedReality/Private/ARTypes.cpp
@@ -2,3 +2,36 @@
 
 
 #include "ARTypes.h"
+#include "HAL/ThreadSafeBool.h"
+
+//@joeg -- ARKit 2.0 support
+
+bool FARAsyncTask::HadError() const
+{
+	return bHadError;
+}
+
+FString FARAsyncTask::GetErrorString() const
+{
+	if (bIsDone)
+	{
+		return Error;
+	}
+	return FString();
+}
+
+bool FARAsyncTask::IsDone() const
+{
+	return bIsDone;
+}
+
+TArray<uint8> FARSaveWorldAsyncTask::GetSavedWorldData()
+{
+	if (bIsDone)
+	{
+		return MoveTemp(WorldData);
+	}
+	return TArray<uint8>();
+}
+
+//@joeg -- End additions
diff --git a/Engine/Source/Runtime/AugmentedReality/Public/ARBlueprintLibrary.h b/Engine/Source/Runtime/AugmentedReality/Public/ARBlueprintLibrary.h
--- a/Engine/Source/Runtime/AugmentedReality/Public/ARBlueprintLibrary.h
+++ b/Engine/Source/Runtime/AugmentedReality/Public/ARBlueprintLibrary.h
@@ -8,6 +8,7 @@
 #include "ARTypes.h"
 #include "ARTraceResult.h"
 #include "ARSessionConfig.h"
+#include "ARTrackable.h"
 #include "ARBlueprintLibrary.generated.h"
 
 
@@ -157,6 +158,37 @@ public:
 	UFUNCTION(BlueprintCallable, Category = "AR AugmentedReality|Pin", meta = (Keywords = "ar augmentedreality augmented reality tracking pin anchor"))
 	static TArray<UARPin*> GetAllPins();
 	
+//@joeg -- Helpers to make things easier
+	/** @return a list of all the tracked planes as currently seen by the Augmented Reality system */
+	UFUNCTION(BlueprintCallable, Category = "AR AugmentedReality|Tracking", meta = (Keywords = "ar augmentedreality augmented reality tracking geometry anchor"))
+	static TArray<UARPlaneGeometry*> GetAllTrackedPlanes();
+	
+	/** @return a list of all the tracked points as currently seen by the Augmented Reality system */
+	UFUNCTION(BlueprintCallable, Category = "AR AugmentedReality|Tracking", meta = (Keywords = "ar augmentedreality augmented reality tracking geometry anchor"))
+	static TArray<UARTrackedPoint*> GetAllTrackedPoints();
+	
+	/** @return a list of all the tracked images as currently seen by the Augmented Reality system */
+	UFUNCTION(BlueprintCallable, Category = "AR AugmentedReality|Tracking", meta = (Keywords = "ar augmentedreality augmented reality tracking images anchor"))
+	static TArray<UARTrackedImage*> GetAllTrackedImages();
+	
+	/** @return a list of all the tracked environment capture probes as currently seen by the Augmented Reality system */
+	UFUNCTION(BlueprintCallable, Category = "AR AugmentedReality|Tracking", meta = (Keywords = "ar augmentedreality augmented reality tracking anchor"))
+	static TArray<UAREnvironmentCaptureProbe*> GetAllTrackedEnvironmentCaptureProbes();
+
+//@joeg -- Added environmental texture probe support
+	/** Adds an environment capture probe to the ar world */
+	UFUNCTION(BlueprintCallable, Category = "AR AugmentedReality|Tracking", meta = (Keywords = "ar augmentedreality augmented reality tracking anchor"))
+	static bool AddManualEnvironmentCaptureProbe(FVector Location, FVector Extent);
+
+	/** @return the current world mapping status for the AR world */
+	UFUNCTION(BlueprintCallable, Category = "AR AugmentedReality|Tracking", meta = (Keywords = "ar augmentedreality augmented reality tracking anchor"))
+	static EARWorldMappingState GetWorldMappingStatus();
+	
+	static TSharedPtr<FARSaveWorldAsyncTask, ESPMode::ThreadSafe> SaveWorld();
+	static TSharedPtr<FARGetCandidateObjectAsyncTask, ESPMode::ThreadSafe> GetCandidateObject(FVector Location, FVector Extent);
+	
+//@joeg -- End additions
+
 public:
 	static void RegisterAsARSystem(const TSharedPtr<FARSystemBase, ESPMode::ThreadSafe>& NewArSystem);
 	
diff --git a/Engine/Source/Runtime/AugmentedReality/Public/ARBlueprintProxy.h b/Engine/Source/Runtime/AugmentedReality/Public/ARBlueprintProxy.h
new file mode 100644
--- /dev/null
+++ b/Engine/Source/Runtime/AugmentedReality/Public/ARBlueprintProxy.h
@@ -0,0 +1,138 @@
+//@joeg -- ARkit 2.0 additions
+
+// Copyright 1998-2018 Epic Games, Inc. All Rights Reserved.
+
+#pragma once
+
+#include "CoreMinimal.h"
+#include "UObject/ObjectMacros.h"
+#include "UObject/Object.h"
+#include "Engine/EngineTypes.h"
+#include "Tickable.h"
+
+#include "ARTypes.h"
+
+#include "ARBlueprintProxy.generated.h"
+
+UCLASS(MinimalAPI)
+class UARBaseAsyncTaskBlueprintProxy :
+	public UObject,
+	public FTickableGameObject
+{
+	GENERATED_UCLASS_BODY()
+	
+public:
+	//~ Begin FTickableObject Interface
+	virtual void Tick(float DeltaTime) override;
+	virtual bool IsTickable() const override { return !HasAnyFlags(RF_ClassDefaultObject) && bShouldTick; }
+	virtual TStatId GetStatId() const override { RETURN_QUICK_DECLARE_CYCLE_STAT(UARBaseAsyncTaskBlueprintProxy, STATGROUP_Tickables); }
+	//~ End FTickableObject Interface
+	
+	/** The async task to check during Tick() */
+	TSharedPtr<FARAsyncTask, ESPMode::ThreadSafe> AsyncTask;
+	
+	virtual void ReportSuccess() { check(0); }
+	virtual void ReportFailure() { check(0); }
+
+	static void RegisterAsARSystem(const TSharedPtr<FARSystemBase, ESPMode::ThreadSafe>& NewArSystem);
+
+protected:
+	static const TSharedPtr<FARSystemBase, ESPMode::ThreadSafe>& GetARSystem();
+
+private:
+	/** True until the async task completes, then false */
+	bool bShouldTick;
+
+	static TSharedPtr<FARSystemBase, ESPMode::ThreadSafe> RegisteredARSystem;
+};
+
+USTRUCT(BlueprintType)
+struct FARSaveWorldResult
+{
+	GENERATED_USTRUCT_BODY()
+	
+	UPROPERTY(BlueprintReadOnly, Category="Augmented Reality")
+	FString Error;
+	
+	UPROPERTY(BlueprintReadOnly, Category="Augmented Reality")
+	TArray<uint8> WorldData;
+};
+
+DECLARE_DYNAMIC_MULTICAST_DELEGATE_OneParam(FARSaveWorldDelegate, const FARSaveWorldResult&, SaveWorldResult);
+
+UCLASS(MinimalAPI)
+class UARSaveWorldAsyncTaskBlueprintProxy :
+	public UARBaseAsyncTaskBlueprintProxy
+{
+	GENERATED_UCLASS_BODY()
+	
+public:
+	UPROPERTY(BlueprintAssignable)
+	FARSaveWorldDelegate OnSuccess;
+	
+	UPROPERTY(BlueprintAssignable)
+	FARSaveWorldDelegate OnFailure;
+	
+	/**
+	 * Saves an AR world to a byte array for network replication or saving to disk
+	 *
+	 * @param bCompressData whether to compress the data or not
+	 */
+	UFUNCTION(BlueprintCallable, Meta=(DisplayName="AR Save World"), Category="Augmented Reality")
+	static UARSaveWorldAsyncTaskBlueprintProxy* CreateProxyObjectForARSaveWorld(bool bCompressData);
+	
+	/** The async task to check during Tick() */
+	TSharedPtr<FARSaveWorldAsyncTask, ESPMode::ThreadSafe> SaveWorldTask;
+	
+	UPROPERTY(BlueprintReadOnly, Category="Augmented Reality")
+	FARSaveWorldResult SaveWorldResult;
+
+	virtual void ReportSuccess() override;
+	virtual void ReportFailure() override;
+};
+
+USTRUCT(BlueprintType)
+struct FARGetCandidateObjectResult
+{
+	GENERATED_USTRUCT_BODY()
+	
+	UPROPERTY(BlueprintReadOnly, Category="Augmented Reality")
+	FString Error;
+	
+	UPROPERTY(BlueprintReadOnly, Category="Augmented Reality")
+	UARCandidateObject* CandidateObject;
+};
+
+DECLARE_DYNAMIC_MULTICAST_DELEGATE_OneParam(FARGetCandidateObjectDelegate, const FARGetCandidateObjectResult&, CandidateObjectResult);
+
+UCLASS(MinimalAPI)
+class UARGetCandidateObjectAsyncTaskBlueprintProxy :
+	public UARBaseAsyncTaskBlueprintProxy
+{
+	GENERATED_UCLASS_BODY()
+	
+public:
+	UPROPERTY(BlueprintAssignable)
+	FARGetCandidateObjectDelegate OnSuccess;
+	
+	UPROPERTY(BlueprintAssignable)
+	FARGetCandidateObjectDelegate OnFailure;
+	
+	/**
+	 * Saves the point cloud centered at the specified location capturing all of the features within the specified extent as an object that can be detected later
+	 *
+	 * @param Location the center of the extent to grab features at
+	 * @param Extent the size of the region to grab feature points
+	 */
+	UFUNCTION(BlueprintCallable, Meta=(DisplayName="AR Get Candidate Object"), Category="Augmented Reality")
+	static UARGetCandidateObjectAsyncTaskBlueprintProxy* CreateProxyObjectForARGetCandidateObject(FVector Location, FVector Extent);
+	
+	/** The async task to check during Tick() */
+	TSharedPtr<FARGetCandidateObjectAsyncTask, ESPMode::ThreadSafe> CandidateObjectTask;
+	
+	UPROPERTY(BlueprintReadOnly, Category="Augmented Reality")
+	FARGetCandidateObjectResult CandidateObjectResult;
+
+	virtual void ReportSuccess() override;
+	virtual void ReportFailure() override;
+};
diff --git a/Engine/Source/Runtime/AugmentedReality/Public/ARSessionConfig.h b/Engine/Source/Runtime/AugmentedReality/Public/ARSessionConfig.h
--- a/Engine/Source/Runtime/AugmentedReality/Public/ARSessionConfig.h
+++ b/Engine/Source/Runtime/AugmentedReality/Public/ARSessionConfig.h
@@ -38,7 +38,15 @@ enum class EARSessionType : uint8
 	World,
 
 	/** AR meant to overlay onto a face */
-	Face
+	Face,
+
+//@joeg -- Added image tracking support
+    /** Tracking of images supplied by the app. No world tracking, just images */
+    Image,
+
+//@joeg -- Object scanning support
+	/** A session used to scan objects for object detection in a world tracking session */
+	ObjectScanning
 };
 
 UENUM(BlueprintType, Category = "AR AugmentedReality", meta = (Experimental, Bitflags))
@@ -77,6 +85,21 @@ enum class EARFrameSyncMode : uint8
 	SyncTickWithoutCameraImage = 1,
 };
 
+//@joeg -- Added texture probe settings
+/**
+ * Tells the AR system what type of environmental texture capturing to perform
+ */
+UENUM(BlueprintType)
+enum class EAREnvironmentCaptureProbeType : uint8
+{
+	/** No capturing will happen */
+	None,
+	/** Capturing will be manual with the app specifying where the probes are and their size */
+	Manual,
+	/** Capturing will be automatic with probes placed by the AR system */
+	Automatic
+};
+
 UCLASS(BlueprintType, Category="AR AugmentedReality")
 class AUGMENTEDREALITY_API UARSessionConfig : public UDataAsset
 {
@@ -113,7 +136,25 @@ public:
 
 	/** @see CandidateImages */
 	const TArray<UARCandidateImage*>& GetCandidateImageList() const;
-
+    
+//@joeg -- Added image tracking support
+	/** @see MaxNumSimultaneousImagesTracked */
+    int32 GetMaxNumSimultaneousImagesTracked() const;
+	
+//@joeg -- Added environmental texture probe support
+	/** @see EnvironmentCaptureProbeType */
+	EAREnvironmentCaptureProbeType GetEnvironmentCaptureProbeType() const;
+	
+//@joeg -- Added for load/save of worlds
+	/** @see WorldMapData */
+	const TArray<uint8>& GetWorldMapData() const;
+	/** @see WorldMapData */
+	void SetWorldMapData(TArray<uint8> WorldMapData);
+
+//@joeg -- For object detection
+	const TArray<UARCandidateObject*>& GetCandidateObjectList() const;
+	void AddCandidateObject(UARCandidateObject* CandidateObject);
+	
 private:
 	//~ UObject interface
 	virtual void Serialize(FArchive& Ar) override;
@@ -163,4 +204,24 @@ protected:
 	/** The list of candidate images to detect within the AR camera view */
 	UPROPERTY(EditAnywhere, Category="AR Settings")
 	TArray<UARCandidateImage*> CandidateImages;
+
+//@joeg -- Added image tracking support
+    /** The maximum number of images to track at the same time. Defaults to 1 */
+    UPROPERTY(EditAnywhere, Category="AR Settings")
+    int32 MaxNumSimultaneousImagesTracked;
+	
+//@joeg -- Added environmental texture probe support
+	/** How the AR system should handle texture probe capturing */
+	UPROPERTY(EditAnywhere, Category="AR Settings")
+	EAREnvironmentCaptureProbeType EnvironmentCaptureProbeType;
+
+//@joeg -- For loading a saved world
+	/** A previously saved world that is to be loaded when the session starts */
+	UPROPERTY(VisibleAnywhere, Category="AR Settings")
+	TArray<uint8> WorldMapData;
+
+//@joeg -- For object detection
+	/** A list of candidate objects to search for in the scene */
+	UPROPERTY(EditAnywhere, Category="AR Settings")
+	TArray<UARCandidateObject*> CandidateObjects;
 };
diff --git a/Engine/Source/Runtime/AugmentedReality/Public/ARSystem.h b/Engine/Source/Runtime/AugmentedReality/Public/ARSystem.h
--- a/Engine/Source/Runtime/AugmentedReality/Public/ARSystem.h
+++ b/Engine/Source/Runtime/AugmentedReality/Public/ARSystem.h
@@ -120,6 +120,20 @@ public:
 	/** @return the last camera depth information the AR system has seen */
 	virtual UARTextureCameraDepth* OnGetCameraDepth() = 0;
 
+//@joeg -- ARKit 2.0 additions
+	/** Tells the ARSystem to generate a capture probe at the specified location if supported */
+	virtual bool OnAddManualEnvironmentCaptureProbe(FVector Location, FVector Extent) = 0;
+	
+	/** Generates a UARCandidateObject from the point cloud data within the location and its extent using an async task */
+	virtual TSharedPtr<FARGetCandidateObjectAsyncTask, ESPMode::ThreadSafe> OnGetCandidateObject(FVector Location, FVector Extent) const = 0;
+
+	/** Saves the AR world to a byte array using an async task */
+	virtual TSharedPtr<FARSaveWorldAsyncTask, ESPMode::ThreadSafe> OnSaveWorld() const = 0;
+	
+	/** @return the current mapping status */
+	virtual EARWorldMappingState OnGetWorldMappingStatus() const = 0;
+//@joeg -- End additions
+	
 public:
 	virtual ~IARSystemSupport(){}
 };
@@ -179,6 +193,18 @@ public:
 	UARTextureCameraImage* GetCameraImage();
 	/** \see UARBlueprintLibrary::GetCameraDepth() */
 	UARTextureCameraDepth* GetCameraDepth();
+//@joeg -- ARKit 2.0 additions
+	/**\see UARBlueprintLibrary::IsEnvironmentCaptureSupported() */
+	bool IsEnvironmentCaptureSupported() const;
+	/**\see UARBlueprintLibrary::AddEnvironmentCaptureProbe() */
+	bool AddManualEnvironmentCaptureProbe(FVector Location, FVector Extent);
+	/** Creates an async task that will perform the work in the background */
+	TSharedPtr<FARGetCandidateObjectAsyncTask, ESPMode::ThreadSafe> GetCandidateObject(FVector Location, FVector Extent) const;
+	/** Creates an async task that will perform the work in the background */
+	TSharedPtr<FARSaveWorldAsyncTask, ESPMode::ThreadSafe> SaveWorld() const;
+	/** @return the current mapping status */
+	EARWorldMappingState GetWorldMappingStatus() const;
+//@joeg -- End additions
 
 	/** \see UARBlueprintLibrary::GetCurrentLightEstimate() */
 	UARLightEstimate* GetCurrentLightEstimate() const;
diff --git a/Engine/Source/Runtime/AugmentedReality/Public/ARTextures.h b/Engine/Source/Runtime/AugmentedReality/Public/ARTextures.h
--- a/Engine/Source/Runtime/AugmentedReality/Public/ARTextures.h
+++ b/Engine/Source/Runtime/AugmentedReality/Public/ARTextures.h
@@ -5,13 +5,16 @@
 #include "CoreMinimal.h"
 #include "UObject/ObjectMacros.h"
 #include "Engine/Texture.h"
+#include "Engine/TextureCube.h"
 #include "ARTextures.generated.h"
 
 UENUM(BlueprintType, Category="AR AugmentedReality", meta=(Experimental))
 enum class EARTextureType : uint8
 {
 	CameraImage,
-	CameraDepth
+	CameraDepth,
+//@joeg -- Addedd environment capture support
+	EnvironmentCapture
 };
 
 /**
@@ -93,3 +96,30 @@ public:
 	UPROPERTY(BlueprintReadOnly, Category="AR AugmentedReality", meta=(Experimental))
 	bool bIsTemporallySmoothed;
 };
+
+//@joeg -- Added for environment capture support
+/**
+ * Base class for all AR textures that represent the environment for lighting and reflection
+ */
+UCLASS(Abstract, BlueprintType)
+class AUGMENTEDREALITY_API UAREnvironmentCaptureProbeTexture : public UTextureCube
+{
+	GENERATED_UCLASS_BODY()
+	
+public:
+	/** The type of texture this is */
+	UPROPERTY(BlueprintReadOnly, Category="AR AugmentedReality", meta=(Experimental))
+	EARTextureType TextureType;
+	
+	/** The timestamp this texture was captured at */
+	UPROPERTY(BlueprintReadOnly, Category="AR AugmentedReality", meta=(Experimental))
+	float Timestamp;
+	
+	/** The guid of texture that gets registered as an external texture */
+	UPROPERTY(BlueprintReadOnly, Category="AR AugmentedReality", meta=(Experimental))
+	FGuid ExternalTextureGuid;
+	
+	/** The width and height of the texture */
+	UPROPERTY(BlueprintReadOnly, Category="AR AugmentedReality", meta=(Experimental))
+	FVector2D Size;
+};
diff --git a/Engine/Source/Runtime/AugmentedReality/Public/ARTrackable.h b/Engine/Source/Runtime/AugmentedReality/Public/ARTrackable.h
--- a/Engine/Source/Runtime/AugmentedReality/Public/ARTrackable.h
+++ b/Engine/Source/Runtime/AugmentedReality/Public/ARTrackable.h
@@ -6,6 +6,7 @@
 #include "ARTrackable.generated.h"
 
 class FARSystemBase;
+class UAREnvironmentCaptureProbeTexture;
 
 UCLASS(BlueprintType)
 class AUGMENTEDREALITY_API UARTrackedGeometry : public UObject
@@ -218,13 +219,31 @@ enum class EARFaceBlendShape : uint8
 	// Nose blend shapes
 	NoseSneerLeft,
 	NoseSneerRight,
+//@joeg -- Tongue blend shape support
+	TongueOut,
 	// Treat the head rotation as curves for LiveLink support
 	HeadYaw,
 	HeadPitch,
 	HeadRoll,
+//@joeg -- Eye tracking support
+	// Treat eye rotation as curves for LiveLink support
+	LeftEyeYaw,
+	LeftEyePitch,
+	LeftEyeRoll,
+	RightEyeYaw,
+	RightEyePitch,
+	RightEyeRoll,
 	MAX
 };
 
+//@joeg -- Eye tracking support
+UENUM(BlueprintType, Category="AR AugmentedReality", meta=(Experimental))
+enum class EAREye : uint8
+{
+	LeftEye,
+	RightEye
+};
+
 typedef TMap<EARFaceBlendShape, float> FARBlendShapeMap;
 
 UCLASS(BlueprintType)
@@ -233,7 +252,7 @@ class AUGMENTEDREALITY_API UARFaceGeometry : public UARTrackedGeometry
 	GENERATED_BODY()
 	
 public:
-	void UpdateTrackedGeometry(const TSharedRef<FARSystemBase, ESPMode::ThreadSafe>& InTrackingSystem, uint32 FrameNumber, double Timestamp, const FTransform& InTransform, const FTransform& InAlignmentTransform, FARBlendShapeMap& InBlendShapes, TArray<FVector>& InVertices, const TArray<int32>& Indices);
+	void UpdateFaceGeometry(const TSharedRef<FARSystemBase, ESPMode::ThreadSafe>& InTrackingSystem, uint32 FrameNumber, double Timestamp, const FTransform& InTransform, const FTransform& InAlignmentTransform, FARBlendShapeMap& InBlendShapes, TArray<FVector>& InVertices, const TArray<int32>& Indices, const FTransform& InLeftEyeTransform, const FTransform& InRightEyeTransform, const FVector& InLookAtTarget);
 	
 	virtual void DebugDraw( UWorld* World, const FLinearColor& OutlineColor, float OutlineThickness, float PersistForSeconds = 0.0f) const override;
 	
@@ -250,6 +269,17 @@ public:
 	const TArray<int32>& GetIndexBuffer() const { return IndexBuffer; }
 	const TArray<FVector2D>& GetUVs() const { return UVs; }
 	
+	UFUNCTION(BlueprintPure, Category="AR AugmentedReality|Face Geometry")
+	const FTransform& GetLocalSpaceEyeTransform(EAREye Eye) const;
+
+	UFUNCTION(BlueprintPure, Category="AR AugmentedReality|Face Geometry")
+	FTransform GetWorldSpaceEyeTransform(EAREye Eye) const;
+
+//@joeg -- Eye tracking support
+	/** The target the eyes are looking at */
+	UPROPERTY(BlueprintReadOnly, Category="AR AugmentedReality|Face Geometry")
+	FVector LookAtTarget;
+
 private:
 	UPROPERTY()
 	TMap<EARFaceBlendShape, float> BlendShapes;
@@ -259,4 +289,63 @@ private:
 	TArray<int32> IndexBuffer;
 	// @todo JoeG - route the uvs in
 	TArray<FVector2D> UVs;
+
+//@joeg -- Eye tracking support
+	/** The transform for the left eye */
+	FTransform LeftEyeTransform;
+	/** The transform for the right eye */
+	FTransform RightEyeTransform;
+};
+
+//@joeg -- Added environmental texture probe support
+/** A tracked environment texture probe that gives you a cube map for reflections */
+UCLASS(BlueprintType)
+class AUGMENTEDREALITY_API UAREnvironmentCaptureProbe :
+	public UARTrackedGeometry
+{
+	GENERATED_BODY()
+	
+public:
+	UAREnvironmentCaptureProbe();
+	
+	/** Draw a box visulizing the bounds of the probe */
+	virtual void DebugDraw(UWorld* World, const FLinearColor& OutlineColor, float OutlineThickness, float PersistForSeconds = 0.0f) const override;
+	
+	void UpdateEnvironmentCapture(const TSharedRef<FARSystemBase, ESPMode::ThreadSafe>& InTrackingSystem, uint32 FrameNumber, double Timestamp, const FTransform& InLocalToTrackingTransform, const FTransform& InAlignmentTransform, FVector InExtent);
+
+	/** @see Extent */
+	UFUNCTION(BlueprintPure, Category="AR AugmentedReality|Environment Capture Probe")
+	FVector GetExtent() const;
+	/** @see EnvironmentCaptureTexture */
+	UFUNCTION(BlueprintPure, Category="AR AugmentedReality|Environment Capture Probe")
+	UAREnvironmentCaptureProbeTexture* GetEnvironmentCaptureTexture() const;
+
+protected:
+	/** The size of area this probe covers */
+	FVector Extent;
+
+	/** The cube map of the reflected environment */
+	UPROPERTY()
+	UAREnvironmentCaptureProbeTexture* EnvironmentCaptureTexture;
+};
+
+//@joeg -- Added for object detection
+UCLASS(BlueprintType)
+class AUGMENTEDREALITY_API UARTrackedObject : public UARTrackedGeometry
+{
+	GENERATED_BODY()
+	
+public:
+	virtual void DebugDraw(UWorld* World, const FLinearColor& OutlineColor, float OutlineThickness, float PersistForSeconds = 0.0f) const override;
+	
+	void UpdateTrackedGeometry(const TSharedRef<FARSystemBase, ESPMode::ThreadSafe>& InTrackingSystem, uint32 FrameNumber, double Timestamp, const FTransform& InLocalToTrackingTransform, const FTransform& InAlignmentTransform, UARCandidateObject* InDetectedObject);
+	
+	/** @see DetectedObject */
+	UFUNCTION(BlueprintPure, Category = "AR AugmentedReality|Object Detection")
+	UARCandidateObject* GetDetectedObject() const { return DetectedObject; };
+	
+private:
+	/** The candidate object that was detected in the scene */
+	UPROPERTY()
+	UARCandidateObject* DetectedObject;
 };
diff --git a/Engine/Source/Runtime/AugmentedReality/Public/ARTypes.h b/Engine/Source/Runtime/AugmentedReality/Public/ARTypes.h
--- a/Engine/Source/Runtime/AugmentedReality/Public/ARTypes.h
+++ b/Engine/Source/Runtime/AugmentedReality/Public/ARTypes.h
@@ -3,6 +3,8 @@
 #pragma once
 
 #include "Engine/DataAsset.h"
+//@joeg -- ARKit 2.0 support
+#include "HAL/ThreadSafeBool.h"
 #include "ARTypes.generated.h"
 
 class FARSystemBase;
@@ -88,6 +90,22 @@ enum class EARSessionStatus : uint8
 	Other,
 };
 
+//@joeg -- ARKit 2.0 support
+/** Gives feedback as to whether the AR data can be saved and relocalized or not */
+UENUM(BlueprintType, Category="AR AugmentedReality", meta=(Experimental))
+enum class EARWorldMappingState : uint8
+{
+	/** World mapping is not available */
+	NotAvailable,
+	/** World mapping is still in progress but without enough data for relocalization */
+	StillMappingNotRelocalizable,
+	/** World mapping is still in progress but there is enough data captured for relocalization */
+	StillMappingRelocalizable,
+	/** World mapping has mapped the area and is fully relocalizable */
+	Mapped
+};
+
+
 /** The current state of the AR subsystem including an optional explanation string. */
 USTRUCT(BlueprintType)
 struct AUGMENTEDREALITY_API FARSessionStatus
@@ -184,3 +202,103 @@ private:
 	UPROPERTY(EditAnywhere, Category = "AR Candidate Image")
 	EARCandidateImageOrientation Orientation;
 };
+
+//@joeg -- ARKit 2.0 support
+
+/** An asset that points to an object to be detected in a scene */
+UCLASS(BlueprintType)
+class AUGMENTEDREALITY_API UARCandidateObject :
+	public UDataAsset
+{
+	GENERATED_BODY()
+	
+public:
+	/** @see CandidateObjectData */
+	const TArray<uint8>& GetCandidateObjectData() const { return CandidateObjectData; }
+	void SetCandidateObjectData(const TArray<uint8>& InCandidateObject) { CandidateObjectData = InCandidateObject; }
+	/** @see FriendlyName */
+	const FString& GetFriendlyName() const { return FriendlyName; }
+	/** @see BoundingBox */
+	const FBox& GetBoundingBox() const { return BoundingBox; }
+	void SetBoundingBox(FBox& InBoundingBox) { BoundingBox = InBoundingBox; }
+
+private:
+	/** The object to detect in scenes */
+	UPROPERTY(EditAnywhere, Category = "AR Candidate Object")
+	TArray<uint8> CandidateObjectData;
+	
+	/** The friendly name to report back when the image is detected in scenes */
+	UPROPERTY(EditAnywhere, Category = "AR Candidate Object")
+	FString FriendlyName;
+	
+	/** The physical width in centimeters of the object that this candidate image represents */
+	UPROPERTY(EditAnywhere, Category = "AR Candidate Image")
+	FBox BoundingBox;
+};
+
+/**
+ * Base class for async AR requests
+ */
+class AUGMENTEDREALITY_API FARAsyncTask
+{
+public:
+	virtual ~FARAsyncTask() {}
+	
+	/** @return whether the task succeeded or not */
+	bool HadError() const;
+	/** @return information about the error if there was one */
+	FString GetErrorString() const;
+	/** @return whether the task has completed or not */
+	bool IsDone() const;
+	
+protected:
+	FThreadSafeBool bIsDone;
+	FThreadSafeBool bHadError;
+	FString Error;
+};
+
+/** Async task that saves the world data into a buffer */
+class AUGMENTEDREALITY_API FARSaveWorldAsyncTask :
+	public FARAsyncTask
+{
+public:
+	/** @return the byte array that the world was saved into. Note uses MoveTemp() for efficiency so only valid once */
+	TArray<uint8> GetSavedWorldData();
+	
+protected:
+	TArray<uint8> WorldData;
+};
+
+/** Async task that builds a candidate object used for detection from the ar session */
+class AUGMENTEDREALITY_API FARGetCandidateObjectAsyncTask :
+	public FARAsyncTask
+{
+public:
+	/** @return the candidate object that you can use for detection later */
+	virtual UARCandidateObject* GetCandidateObject() = 0;
+};
+
+class FARErrorGetCandidateObjectAsyncTask :
+	public FARGetCandidateObjectAsyncTask
+{
+public:
+	FARErrorGetCandidateObjectAsyncTask(FString InError)
+	{
+		Error = InError;
+		bHadError = true;
+		bIsDone = true;
+	}
+	virtual UARCandidateObject* GetCandidateObject() override { return nullptr; }
+};
+
+class FARErrorSaveWorldAsyncTask :
+	public FARSaveWorldAsyncTask
+{
+public:
+	FARErrorSaveWorldAsyncTask(FString InError)
+	{
+		Error = InError;
+		bHadError = true;
+		bIsDone = true;
+	}
+};
