Fixed missing SUPPORTS_ARKIT_*  wrapper

diff --git a/Engine/Plugins/Runtime/AppleARKit/Source/AppleARKit/Public/AppleARKitConversion.h b/Engine/Plugins/Runtime/AppleARKit/Source/AppleARKit/Public/AppleARKitConversion.h
--- a/Engine/Plugins/Runtime/AppleARKit/Source/AppleARKit/Public/AppleARKitConversion.h
+++ b/Engine/Plugins/Runtime/AppleARKit/Source/AppleARKit/Public/AppleARKitConversion.h
@@ -405,6 +405,7 @@ struct FAppleARKitConversion
 #endif
 };
 
+#if SUPPORTS_ARKIT_1_0
 enum class EAppleAnchorType : uint8
 {
 	Anchor,
@@ -491,6 +492,7 @@ struct FAppleARKitAnchorData
 	FTransform RightEyeTransform;
 	FVector LookAtTarget;
 };
+#endif
 
 namespace ARKitUtil
 {
