Added support to query and specify the desired video format for an AR session

diff --git a/Engine/Plugins/Runtime/AppleARKit/Source/AppleARKit/Private/AppleARKitSystem.cpp b/Engine/Plugins/Runtime/AppleARKit/Source/AppleARKit/Private/AppleARKitSystem.cpp
--- a/Engine/Plugins/Runtime/AppleARKit/Source/AppleARKit/Private/AppleARKitSystem.cpp
+++ b/Engine/Plugins/Runtime/AppleARKit/Source/AppleARKit/Private/AppleARKitSystem.cpp
@@ -769,6 +769,29 @@ bool FAppleARKitSystem::OnAddManualEnvironmentCaptureProbe(FVector Location, FVe
 	return false;
 }
 
+TArray<FARVideoFormat> FAppleARKitSystem::OnGetSupportedVideoFormats(EARSessionType SessionType) const
+{
+#if SUPPORTS_ARKIT_1_5
+	if (FAppleARKitAvailability::SupportsARKit15())
+	{
+		switch (SessionType)
+		{
+			case EARSessionType::Face:
+			{
+				// We need to get the face support from the factory method, which is a modular feature to avoid dependencies
+				TArray<IAppleARKitFaceSupport*> Impls = IModularFeatures::Get().GetModularFeatureImplementations<IAppleARKitFaceSupport>("AppleARKitFaceSupport");
+				break;
+			}
+			case EARSessionType::World:
+			{
+				return FAppleARKitConversion::FromARVideoFormatArray(ARWorldTrackingConfiguration.supportedVideoFormats);
+			}
+		}
+	}
+#endif
+	return TArray<FARVideoFormat>();
+}
+
 #if SUPPORTS_ARKIT_2_0
 /** Since both the object extraction and world saving need to get the world map async, use a common chunk of code for this */
 class FAppleARKitGetWorldMapObjectAsyncTask
diff --git a/Engine/Plugins/Runtime/AppleARKit/Source/AppleARKit/Private/AppleARKitSystem.h b/Engine/Plugins/Runtime/AppleARKit/Source/AppleARKit/Private/AppleARKitSystem.h
--- a/Engine/Plugins/Runtime/AppleARKit/Source/AppleARKit/Private/AppleARKitSystem.h
+++ b/Engine/Plugins/Runtime/AppleARKit/Source/AppleARKit/Private/AppleARKitSystem.h
@@ -85,6 +85,7 @@ protected:
 	virtual TSharedPtr<FARSaveWorldAsyncTask, ESPMode::ThreadSafe> OnSaveWorld() const override;
 	virtual EARWorldMappingState OnGetWorldMappingStatus() const override;
 //@joeg -- End additions
+	virtual TArray<FARVideoFormat> OnGetSupportedVideoFormats(EARSessionType SessionType) const override;
 	//~IARSystemSupport
 
 private:
diff --git a/Engine/Plugins/Runtime/AppleARKit/Source/AppleARKit/Public/AppleARKitConversion.h b/Engine/Plugins/Runtime/AppleARKit/Source/AppleARKit/Public/AppleARKitConversion.h
--- a/Engine/Plugins/Runtime/AppleARKit/Source/AppleARKit/Public/AppleARKitConversion.h
+++ b/Engine/Plugins/Runtime/AppleARKit/Source/AppleARKit/Public/AppleARKitConversion.h
@@ -151,6 +151,52 @@ struct FAppleARKitConversion
 
 	#if SUPPORTS_ARKIT_1_5
 
+	static ARVideoFormat* ToARVideoFormat(const FARVideoFormat& DesiredFormat, NSArray<ARVideoFormat*>* Formats)
+	{
+		if (Formats != nullptr)
+		{
+			for (ARVideoFormat* Format in Formats)
+			{
+				if (Format != nullptr &&
+					DesiredFormat.FPS == Format.framesPerSecond &&
+					DesiredFormat.Width == Format.imageResolution.width &&
+					DesiredFormat.Height == Format.imageResolution.height)
+				{
+					return Format;
+				}
+			}
+		}
+		return nullptr;
+	}
+	
+	static FARVideoFormat FromARVideoFormat(ARVideoFormat* Format)
+	{
+		FARVideoFormat ConvertedFormat;
+		if (Format != nullptr)
+		{
+			ConvertedFormat.FPS = Format.framesPerSecond;
+			ConvertedFormat.Width = Format.imageResolution.width;
+			ConvertedFormat.Height = Format.imageResolution.height;
+		}
+		return ConvertedFormat;
+	}
+	
+	static TArray<FARVideoFormat> FromARVideoFormatArray(NSArray<ARVideoFormat*>* Formats)
+	{
+		TArray<FARVideoFormat> ConvertedArray;
+		if (Formats != nullptr)
+		{
+			for (ARVideoFormat* Format in Formats)
+			{
+				if (Format != nullptr)
+				{
+					ConvertedArray.Add(FromARVideoFormat(Format));
+				}
+			}
+		}
+		return ConvertedArray;
+	}
+
 	static NSSet* InitImageDetection(UARSessionConfig* SessionConfig, TMap< FString, UARCandidateImage* >& CandidateImages, TMap< FString, CGImageRef >& ConvertedCandidateImages)
 	{
 		const TArray<UARCandidateImage*>& ConfigCandidateImages = SessionConfig->GetCandidateImageList();
@@ -327,6 +373,11 @@ struct FAppleARKitConversion
 					WorldTrackingConfiguration.autoFocusEnabled = SessionConfig->ShouldEnableAutoFocus();
 					// Add any images that wish to be detected
 					FAppleARKitConversion::InitImageDetection(SessionConfig, WorldTrackingConfiguration, CandidateImages, ConvertedCandidateImages);
+					ARVideoFormat* Format = FAppleARKitConversion::ToARVideoFormat(SessionConfig->GetDesiredVideoFormat(), ARWorldTrackingConfiguration.supportedVideoFormats);
+					if (Format != nullptr)
+					{
+						WorldTrackingConfiguration.videoFormat = Format;
+					}
 				}
 #endif
 //@joeg -- Added environmental texture probe support
@@ -530,4 +581,4 @@ namespace ARKitUtil
 
 		return OutPins;
 	}
-}
\ No newline at end of file
+}
diff --git a/Engine/Plugins/Runtime/AppleARKit/Source/AppleARKit/Public/AppleARKitFaceSupport.h b/Engine/Plugins/Runtime/AppleARKit/Source/AppleARKit/Public/AppleARKitFaceSupport.h
--- a/Engine/Plugins/Runtime/AppleARKit/Source/AppleARKit/Public/AppleARKitFaceSupport.h
+++ b/Engine/Plugins/Runtime/AppleARKit/Source/AppleARKit/Public/AppleARKitFaceSupport.h
@@ -13,7 +13,7 @@
 
 class UARSessionConfig;
 struct FAppleARKitAnchorData;
-
+struct FARVideoFormat;
 
 class APPLEARKIT_API IAppleARKitFaceSupport :
 	public IModularFeature
@@ -51,6 +51,12 @@ public:
 	 */
 	virtual ARConfiguration* ToARConfiguration(UARSessionConfig* SessionConfig) { return nullptr; }
 #endif
+#if SUPPORTS_ARKIT_1_5
+	/**
+	 * @return the supported video formats by the face ar device
+	 */
+	virtual TArray<FARVideoFormat> ToARConfiguration() { return TArray<FARVideoFormat>(); }
+#endif
 
 	static FName GetModularFeatureName()
 	{
diff --git a/Engine/Plugins/Runtime/AppleARKitFaceSupport/Source/AppleARKitFaceSupport/Private/AppleARKitFaceSupportImpl.cpp b/Engine/Plugins/Runtime/AppleARKitFaceSupport/Source/AppleARKitFaceSupport/Private/AppleARKitFaceSupportImpl.cpp
--- a/Engine/Plugins/Runtime/AppleARKitFaceSupport/Source/AppleARKitFaceSupport/Private/AppleARKitFaceSupportImpl.cpp
+++ b/Engine/Plugins/Runtime/AppleARKitFaceSupport/Source/AppleARKitFaceSupport/Private/AppleARKitFaceSupportImpl.cpp
@@ -76,7 +76,7 @@ void FAppleARKitFaceSupport::Shutdown()
 
 ARConfiguration* FAppleARKitFaceSupport::ToARConfiguration(UARSessionConfig* SessionConfig)
 {
-	ARConfiguration* SessionConfiguration = nullptr;
+	ARFaceTrackingConfiguration* SessionConfiguration = nullptr;
 	if (SessionConfig->GetSessionType() == EARSessionType::Face)
 	{
 		if (ARFaceTrackingConfiguration.isSupported == FALSE)
@@ -91,6 +91,16 @@ ARConfiguration* FAppleARKitFaceSupport::ToARConfiguration(UARSessionConfig* Ses
 	SessionConfiguration.providesAudioData = NO;
 	SessionConfiguration.worldAlignment = FAppleARKitConversion::ToARWorldAlignment(SessionConfig->GetWorldAlignment());
 
+#if SUPPORTS_ARKIT_1_5
+	if (FAppleARKitAvailability::SupportsARKit15())
+	{
+		ARVideoFormat* Format = FAppleARKitConversion::ToARVideoFormat(SessionConfig->GetDesiredVideoFormat(), ARFaceTrackingConfiguration.supportedVideoFormats);
+		if (Format != nullptr)
+		{
+			SessionConfiguration.videoFormat = Format;
+		}
+	}
+#endif
 	return SessionConfiguration;
 }
 
@@ -136,3 +146,10 @@ void FAppleARKitFaceSupport::PublishLiveLinkData(TSharedPtr<FAppleARKitAnchorDat
 }
 
 #endif
+#if SUPPORTS_ARKIT_1_5
+TArray<FARVideoFormat> FAppleARKitFaceSupport::ToARConfiguration()
+{
+	return FAppleARKitConversion::FromARVideoFormatArray(ARFaceTrackingConfiguration.supportedVideoFormats);
+}
+#endif
+
diff --git a/Engine/Plugins/Runtime/AppleARKitFaceSupport/Source/AppleARKitFaceSupport/Private/AppleARKitFaceSupportImpl.h b/Engine/Plugins/Runtime/AppleARKitFaceSupport/Source/AppleARKitFaceSupport/Private/AppleARKitFaceSupportImpl.h
--- a/Engine/Plugins/Runtime/AppleARKitFaceSupport/Source/AppleARKitFaceSupport/Private/AppleARKitFaceSupportImpl.h
+++ b/Engine/Plugins/Runtime/AppleARKitFaceSupport/Source/AppleARKitFaceSupport/Private/AppleARKitFaceSupportImpl.h
@@ -25,6 +25,9 @@ private:
 	virtual TArray<TSharedPtr<FAppleARKitAnchorData>> MakeAnchorData(NSArray<ARAnchor*>* NewAnchors, double Timestamp, uint32 FrameNumber, const FRotator& AdjustBy) override;
 	virtual ARConfiguration* ToARConfiguration(UARSessionConfig* InSessionConfig) override;
 	virtual void PublishLiveLinkData(TSharedPtr<FAppleARKitAnchorData> Anchor, double Timestamp, uint32 FrameNumber) override;
+#if SUPPORTS_ARKIT_1_5
+	virtual TArray<FARVideoFormat> ToARConfiguration() override;
+#endif
 	// ~IAppleARKitFaceSupport
 #endif
 
diff --git a/Engine/Plugins/Runtime/GoogleARCore/Source/GoogleARCoreBase/Private/GoogleARCoreXRTrackingSystem.h b/Engine/Plugins/Runtime/GoogleARCore/Source/GoogleARCoreBase/Private/GoogleARCoreXRTrackingSystem.h
--- a/Engine/Plugins/Runtime/GoogleARCore/Source/GoogleARCoreBase/Private/GoogleARCoreXRTrackingSystem.h
+++ b/Engine/Plugins/Runtime/GoogleARCore/Source/GoogleARCoreBase/Private/GoogleARCoreXRTrackingSystem.h
@@ -60,11 +60,12 @@ protected:
 	virtual void OnRemovePin(UARPin* PinToRemove) override;
 	virtual UARTextureCameraImage* OnGetCameraImage() override { return nullptr; }
 	virtual UARTextureCameraDepth* OnGetCameraDepth() override { return nullptr; }
-	virtual bool OnAddManualEnvironmentCaptureProbe(FVector Location, FVector Extent) { return false; }
-	virtual TSharedPtr<FARGetCandidateObjectAsyncTask, ESPMode::ThreadSafe> OnGetCandidateObject(FVector Location, FVector Extent) const { return TSharedPtr<FARGetCandidateObjectAsyncTask, ESPMode::ThreadSafe>(); }
-	virtual TSharedPtr<FARSaveWorldAsyncTask, ESPMode::ThreadSafe> OnSaveWorld() const { return TSharedPtr<FARSaveWorldAsyncTask, ESPMode::ThreadSafe>(); }
-// @todo -- support this properly
-	virtual EARWorldMappingState OnGetWorldMappingStatus() const { return EARWorldMappingState::StillMappingNotRelocalizable; }
+	virtual bool OnAddManualEnvironmentCaptureProbe(FVector Location, FVector Extent) override { return false; }
+	virtual TSharedPtr<FARGetCandidateObjectAsyncTask, ESPMode::ThreadSafe> OnGetCandidateObject(FVector Location, FVector Extent) const override { return TSharedPtr<FARGetCandidateObjectAsyncTask, ESPMode::ThreadSafe>(); }
+	virtual TSharedPtr<FARSaveWorldAsyncTask, ESPMode::ThreadSafe> OnSaveWorld() const override { return TSharedPtr<FARSaveWorldAsyncTask, ESPMode::ThreadSafe>(); }
+// @todo -- support these properly
+	virtual EARWorldMappingState OnGetWorldMappingStatus() const override { return EARWorldMappingState::StillMappingNotRelocalizable; }
+	virtual TArray<FARVideoFormat> OnGetSupportedVideoFormats(EARSessionType SessionType) const override { return TArray<FARVideoFormat>(); }
 	//~IARSystemSupport
 
 private:
diff --git a/Engine/Source/Runtime/AugmentedReality/Private/ARBlueprintLibrary.cpp b/Engine/Source/Runtime/AugmentedReality/Private/ARBlueprintLibrary.cpp
--- a/Engine/Source/Runtime/AugmentedReality/Private/ARBlueprintLibrary.cpp
+++ b/Engine/Source/Runtime/AugmentedReality/Private/ARBlueprintLibrary.cpp
@@ -458,3 +458,13 @@ TSharedPtr<FARGetCandidateObjectAsyncTask, ESPMode::ThreadSafe> UARBlueprintLibr
 }
 
 //@joeg -- End additions
+
+TArray<FARVideoFormat> UARBlueprintLibrary::GetSupportedVideoFormats(EARSessionType SessionType)
+{
+	auto ARSystem = GetARSystem();
+	if (ensure(ARSystem.IsValid()))
+	{
+		return ARSystem->GetSupportedVideoFormats(SessionType);
+	}
+	return TArray<FARVideoFormat>();
+}
diff --git a/Engine/Source/Runtime/AugmentedReality/Private/ARSessionConfig.cpp b/Engine/Source/Runtime/AugmentedReality/Private/ARSessionConfig.cpp
--- a/Engine/Source/Runtime/AugmentedReality/Private/ARSessionConfig.cpp
+++ b/Engine/Source/Runtime/AugmentedReality/Private/ARSessionConfig.cpp
@@ -106,6 +106,16 @@ void UARSessionConfig::AddCandidateObject(UARCandidateObject* CandidateObject)
 }
 //@joeg -- end object detection
 
+FARVideoFormat UARSessionConfig::GetDesiredVideoFormat() const
+{
+	return DesiredVideoFormat;
+}
+
+void UARSessionConfig::SetDesiredVideoFormat(FARVideoFormat NewFormat)
+{
+	DesiredVideoFormat = NewFormat;
+}
+
 void UARSessionConfig::Serialize(FArchive& Ar)
 {
 	Ar.UsingCustomVersion(FVRObjectVersion::GUID);
@@ -120,4 +130,4 @@ void UARSessionConfig::Serialize(FArchive& Ar)
 			bVerticalPlaneDetection = false;
 		}
 	}
-}
\ No newline at end of file
+}
diff --git a/Engine/Source/Runtime/AugmentedReality/Private/ARSystem.cpp b/Engine/Source/Runtime/AugmentedReality/Private/ARSystem.cpp
--- a/Engine/Source/Runtime/AugmentedReality/Private/ARSystem.cpp
+++ b/Engine/Source/Runtime/AugmentedReality/Private/ARSystem.cpp
@@ -132,6 +132,10 @@ void FARSystemBase::SetAlignmentTransform( const FTransform& InAlignmentTransfor
 
 }
 
+TArray<FARVideoFormat> FARSystemBase::GetSupportedVideoFormats(EARSessionType SessionType) const
+{
+	return OnGetSupportedVideoFormats(SessionType);
+}
 
 UARLightEstimate* FARSystemBase::GetCurrentLightEstimate() const
 {
diff --git a/Engine/Source/Runtime/AugmentedReality/Public/ARBlueprintLibrary.h b/Engine/Source/Runtime/AugmentedReality/Public/ARBlueprintLibrary.h
--- a/Engine/Source/Runtime/AugmentedReality/Public/ARBlueprintLibrary.h
+++ b/Engine/Source/Runtime/AugmentedReality/Public/ARBlueprintLibrary.h
@@ -184,6 +184,10 @@ public:
 	UFUNCTION(BlueprintCallable, Category = "AR AugmentedReality|Tracking", meta = (Keywords = "ar augmentedreality augmented reality tracking anchor"))
 	static EARWorldMappingState GetWorldMappingStatus();
 	
+	/** @return The list of supported video formats for this device */
+	UFUNCTION(BlueprintCallable, Category = "AR AugmentedReality|Capabilities", meta = (Keywords = "ar augmentedreality augmented reality config video formats"))
+	static TArray<FARVideoFormat> GetSupportedVideoFormats(EARSessionType SessionType);
+
 	static TSharedPtr<FARSaveWorldAsyncTask, ESPMode::ThreadSafe> SaveWorld();
 	static TSharedPtr<FARGetCandidateObjectAsyncTask, ESPMode::ThreadSafe> GetCandidateObject(FVector Location, FVector Extent);
 	
diff --git a/Engine/Source/Runtime/AugmentedReality/Public/ARSessionConfig.h b/Engine/Source/Runtime/AugmentedReality/Public/ARSessionConfig.h
--- a/Engine/Source/Runtime/AugmentedReality/Public/ARSessionConfig.h
+++ b/Engine/Source/Runtime/AugmentedReality/Public/ARSessionConfig.h
@@ -155,6 +155,10 @@ public:
 	const TArray<UARCandidateObject*>& GetCandidateObjectList() const;
 	void AddCandidateObject(UARCandidateObject* CandidateObject);
 	
+	/** @see DesiredVideoFormat */
+	FARVideoFormat GetDesiredVideoFormat() const;
+	void SetDesiredVideoFormat(FARVideoFormat NewFormat);
+
 private:
 	//~ UObject interface
 	virtual void Serialize(FArchive& Ar) override;
@@ -224,4 +228,11 @@ protected:
 	/** A list of candidate objects to search for in the scene */
 	UPROPERTY(EditAnywhere, Category="AR Settings")
 	TArray<UARCandidateObject*> CandidateObjects;
+
+	/**
+	 * The desired video format (or the default if not supported) that this session should use if the camera is enabled
+	 * Note: Call GetSupportedVideoFormats to get a list of device supported formats
+	 */
+	UPROPERTY(EditAnywhere, Category="AR Settings")
+	FARVideoFormat DesiredVideoFormat;
 };
diff --git a/Engine/Source/Runtime/AugmentedReality/Public/ARSystem.h b/Engine/Source/Runtime/AugmentedReality/Public/ARSystem.h
--- a/Engine/Source/Runtime/AugmentedReality/Public/ARSystem.h
+++ b/Engine/Source/Runtime/AugmentedReality/Public/ARSystem.h
@@ -134,6 +134,9 @@ public:
 	virtual EARWorldMappingState OnGetWorldMappingStatus() const = 0;
 //@joeg -- End additions
 	
+	/** @return The list of supported video formats for this device and session type */
+	virtual TArray<FARVideoFormat> OnGetSupportedVideoFormats(EARSessionType SessionType) const = 0;
+	
 public:
 	virtual ~IARSystemSupport(){}
 };
@@ -216,6 +219,9 @@ public:
 	/** \see UARBlueprintLibrary::RemovePin() */
 	void RemovePin( UARPin* PinToRemove );
 
+	/** \see UARBlueprintLibrary::GetSupportedVideoFormats() */
+	TArray<FARVideoFormat> GetSupportedVideoFormats(EARSessionType SessionType = EARSessionType::World) const;
+	
 	virtual void* GetARSessionRawPointer() = 0;
 	virtual void* GetGameThreadARFrameRawPointer() = 0;
 	
diff --git a/Engine/Source/Runtime/AugmentedReality/Public/ARTypes.h b/Engine/Source/Runtime/AugmentedReality/Public/ARTypes.h
--- a/Engine/Source/Runtime/AugmentedReality/Public/ARTypes.h
+++ b/Engine/Source/Runtime/AugmentedReality/Public/ARTypes.h
@@ -320,3 +320,36 @@ public:
 		bIsDone = true;
 	}
 };
+
+/** A specific AR video format */
+USTRUCT(BlueprintType)
+struct AUGMENTEDREALITY_API FARVideoFormat
+{
+	GENERATED_BODY()
+
+public:
+	FARVideoFormat() { }
+	
+	FARVideoFormat(int32 InFPS, int32 InWidth, int32 InHeight)
+		: FPS(InFPS)
+		, Width(InWidth)
+		, Height(InHeight)
+	{
+		
+	}
+	
+	/** The desired or supported number of frames per second for this video format */
+	UPROPERTY(BlueprintReadOnly, Category="AR AugmentedReality|Session")
+	int32 FPS;
+	
+	/** The desired or supported width in pixels for this video format */
+	UPROPERTY(BlueprintReadOnly, Category = "AR AugmentedReality|Session")
+	int32 Width;
+
+	/** The desired or supported height in pixels for this video format */
+	UPROPERTY(BlueprintReadOnly, Category = "AR AugmentedReality|Session")
+	int32 Height;
+
+	bool IsValidFormat() { return FPS > 0 && Width > 0 && Height > 0; }
+};
+
