Removed all of my merge markers from the lab

diff --git a/Engine/Plugins/Experimental/AppleImageUtils/Source/AppleImageUtils/Public/AppleImageUtilsTypes.h b/Engine/Plugins/Experimental/AppleImageUtils/Source/AppleImageUtils/Public/AppleImageUtilsTypes.h
--- a/Engine/Plugins/Experimental/AppleImageUtils/Source/AppleImageUtils/Public/AppleImageUtilsTypes.h
+++ b/Engine/Plugins/Experimental/AppleImageUtils/Source/AppleImageUtils/Public/AppleImageUtilsTypes.h
@@ -47,7 +47,6 @@ public:
 	/** @return the IOSurface held by the implementing object */
 	virtual IOSurfaceRef GetSurface() const { return nullptr; }
 
-//@joeg -- Added for environment capture support
 	/** @return the MTLTextureid<MTL held by the implementing object */
 	virtual id<MTLTexture> GetMetalTexture() const { return nullptr; }
 #endif
diff --git a/Engine/Plugins/Runtime/AppleARKit/Source/AppleARKit/Private/AppleARKitSystem.cpp b/Engine/Plugins/Runtime/AppleARKit/Source/AppleARKit/Private/AppleARKitSystem.cpp
--- a/Engine/Plugins/Runtime/AppleARKit/Source/AppleARKit/Private/AppleARKitSystem.cpp
+++ b/Engine/Plugins/Runtime/AppleARKit/Source/AppleARKit/Private/AppleARKitSystem.cpp
@@ -14,7 +14,6 @@
 #include "GeneralProjectSettings.h"
 #include "ARSessionConfig.h"
 #include "AppleARKitSettings.h"
-//@joeg -- Environment capture support
 #include "AppleARKitTrackable.h"
 #include "ARLightEstimate.h"
 #include "ARTraceResult.h"
@@ -752,8 +751,6 @@ bool FAppleARKitSystem::OnIsTrackingTypeSupported(EARSessionType SessionType) co
 	return false;
 }
 
-//@joeg -- ARKit 2.0 additions
-
 bool FAppleARKitSystem::OnAddManualEnvironmentCaptureProbe(FVector Location, FVector Extent)
 {
 #if SUPPORTS_ARKIT_2_0
@@ -1008,7 +1005,6 @@ EARWorldMappingState FAppleARKitSystem::OnGetWorldMappingStatus() const
 	return EARWorldMappingState::NotAvailable;
 }
 
-//@joeg -- End additions
 
 void FAppleARKitSystem::AddReferencedObjects( FReferenceCollector& Collector )
 {
@@ -1019,7 +1015,6 @@ void FAppleARKitSystem::AddReferencedObjects( FReferenceCollector& Collector )
 	Collector.AddReferencedObject( CameraImage );
 	Collector.AddReferencedObject( CameraDepth );
 	Collector.AddReferencedObjects( CandidateImages );
-//@joeg -- Object detection
 	Collector.AddReferencedObjects( CandidateObjects );
 
 	if(LightEstimate)
@@ -1115,7 +1110,6 @@ bool FAppleARKitSystem::Run(UARSessionConfig* SessionConfig)
 		return true;
 	}
 
-//@joeg -- bug fix
 	{
 		// Clear out any existing frames since they aren't valid anymore
 		FScopeLock ScopeLock(&FrameLock);
@@ -1305,7 +1299,6 @@ static TSharedPtr<FAppleARKitAnchorData> MakeAnchorData( ARAnchor* Anchor )
 		NewAnchor = MakeShared<FAppleARKitAnchorData>(
 			FAppleARKitConversion::ToFGuid(ImageAnchor.identifier),
 			FAppleARKitConversion::ToFTransform(ImageAnchor.transform),
-//@joeg -- Changed to share image and object detection
 			EAppleAnchorType::ImageAnchor,
 			FString(ImageAnchor.referenceImage.name)
 		);
@@ -1314,7 +1307,6 @@ static TSharedPtr<FAppleARKitAnchorData> MakeAnchorData( ARAnchor* Anchor )
 #endif
 	}
 #endif
-//@joeg -- Added environmental texture probe support
 #if SUPPORTS_ARKIT_2_0
 	else if (FAppleARKitAvailability::SupportsARKit20() && [Anchor isKindOfClass:[AREnvironmentProbeAnchor class]])
 	{
@@ -1326,7 +1318,6 @@ static TSharedPtr<FAppleARKitAnchorData> MakeAnchorData( ARAnchor* Anchor )
 			ProbeAnchor.environmentTexture
 		);
 	}
-//@joeg -- Object dectection support
 	else if (FAppleARKitAvailability::SupportsARKit20() && [Anchor isKindOfClass:[ARObjectAnchor class]])
 	{
 		ARObjectAnchor* ObjectAnchor = (ARObjectAnchor*)Anchor;
@@ -1464,7 +1455,6 @@ void FAppleARKitSystem::SessionDidAddAnchors_Internal( TSharedRef<FAppleARKitAnc
 		{
 			NewAnchorDebugName = FString::Printf(TEXT("FACE-%02d"), LastTrackedGeometry_DebugId++);
 			UARFaceGeometry* NewGeo = NewObject<UARFaceGeometry>();
-//@joeg -- Eye tracking support
 			NewGeo->UpdateFaceGeometry(SharedThis(this), GameThreadFrameNumber, GameThreadTimestamp, AnchorData->Transform, GetAlignmentTransform(), AnchorData->BlendShapes, AnchorData->FaceVerts, AnchorData->FaceIndices, AnchorData->LeftEyeTransform, AnchorData->RightEyeTransform, AnchorData->LookAtTarget);
 			NewGeo->bIsTracked = true;
 			NewGeometry = NewGeo;
@@ -1480,7 +1470,6 @@ void FAppleARKitSystem::SessionDidAddAnchors_Internal( TSharedRef<FAppleARKitAnc
 			NewGeometry = NewImage;
 			break;
 		}
-//@joeg -- Added environmental texture probe support
 		case EAppleAnchorType::EnvironmentProbeAnchor:
 		{
 			NewAnchorDebugName = FString::Printf(TEXT("ENV-%02d"), LastTrackedGeometry_DebugId++);
@@ -1489,7 +1478,6 @@ void FAppleARKitSystem::SessionDidAddAnchors_Internal( TSharedRef<FAppleARKitAnc
 			NewGeometry = NewProbe;
 			break;
 		}
-//@joeg -- Object detection
 		case EAppleAnchorType::ObjectAnchor:
 		{
 			NewAnchorDebugName = FString::Printf(TEXT("OBJ-%02d"), LastTrackedGeometry_DebugId++);
@@ -1566,7 +1554,6 @@ void FAppleARKitSystem::SessionDidUpdateAnchors_Internal( TSharedRef<FAppleARKit
 			{
 				if (UARFaceGeometry* FaceGeo = Cast<UARFaceGeometry>(FoundGeometry))
 				{
-//@joeg -- Eye tracking support
 					FaceGeo->UpdateFaceGeometry(SharedThis(this), GameThreadFrameNumber, GameThreadTimestamp, AnchorData->Transform, GetAlignmentTransform(), AnchorData->BlendShapes, AnchorData->FaceVerts, AnchorData->FaceIndices, AnchorData->LeftEyeTransform, AnchorData->RightEyeTransform, AnchorData->LookAtTarget);
 					FaceGeo->bIsTracked = AnchorData->bIsTracked;
 					for (UARPin* Pin : PinsToUpdate)
@@ -1577,12 +1564,10 @@ void FAppleARKitSystem::SessionDidUpdateAnchors_Internal( TSharedRef<FAppleARKit
 				}
 				break;
 			}
-//@joeg -- Added image tracking support
             case EAppleAnchorType::ImageAnchor:
             {
 				if (UARTrackedImage* ImageAnchor = Cast<UARTrackedImage>(FoundGeometry))
 				{
-//@joeg -- Changed so object and image detection can share
 					UARCandidateImage** CandidateImage = CandidateImages.Find(AnchorData->DetectedAnchorName);
 					ensure(CandidateImage != nullptr);
 
@@ -1596,7 +1581,6 @@ void FAppleARKitSystem::SessionDidUpdateAnchors_Internal( TSharedRef<FAppleARKit
 				}
                 break;
             }
-//@joeg -- Added environmental texture probe support
 			case EAppleAnchorType::EnvironmentProbeAnchor:
 			{
 				if (UAppleARKitEnvironmentCaptureProbe* ProbeAnchor = Cast<UAppleARKitEnvironmentCaptureProbe>(FoundGeometry))
diff --git a/Engine/Plugins/Runtime/AppleARKit/Source/AppleARKit/Private/AppleARKitSystem.h b/Engine/Plugins/Runtime/AppleARKit/Source/AppleARKit/Private/AppleARKitSystem.h
--- a/Engine/Plugins/Runtime/AppleARKit/Source/AppleARKit/Private/AppleARKitSystem.h
+++ b/Engine/Plugins/Runtime/AppleARKit/Source/AppleARKit/Private/AppleARKitSystem.h
@@ -79,12 +79,10 @@ protected:
 	virtual void OnRemovePin(UARPin* PinToRemove) override;
 	virtual UARTextureCameraImage* OnGetCameraImage() override;
 	virtual UARTextureCameraDepth* OnGetCameraDepth() override;
-//@joeg -- ARKit 2.0 additions
 	virtual bool OnAddManualEnvironmentCaptureProbe(FVector Location, FVector Extent) override;
 	virtual TSharedPtr<FARGetCandidateObjectAsyncTask, ESPMode::ThreadSafe> OnGetCandidateObject(FVector Location, FVector Extent) const override;
 	virtual TSharedPtr<FARSaveWorldAsyncTask, ESPMode::ThreadSafe> OnSaveWorld() const override;
 	virtual EARWorldMappingState OnGetWorldMappingStatus() const override;
-//@joeg -- End additions
 	virtual TArray<FARVideoFormat> OnGetSupportedVideoFormats(EARSessionType SessionType) const override;
 	//~IARSystemSupport
 
@@ -169,7 +167,6 @@ private:
 	UAppleARKitTextureCameraImage* CameraImage;
 	UAppleARKitTextureCameraDepth* CameraDepth;
 	TMap< FString, UARCandidateImage* > CandidateImages;
-//@joeg -- Object detection
 	TMap< FString, UARCandidateObject* > CandidateObjects;
 	// ...
 	// PROPERTIES REPORTED TO FGCObject
diff --git a/Engine/Plugins/Runtime/AppleARKit/Source/AppleARKit/Private/AppleARKitTextures.cpp b/Engine/Plugins/Runtime/AppleARKit/Source/AppleARKit/Private/AppleARKitTextures.cpp
--- a/Engine/Plugins/Runtime/AppleARKit/Source/AppleARKit/Private/AppleARKitTextures.cpp
+++ b/Engine/Plugins/Runtime/AppleARKit/Source/AppleARKit/Private/AppleARKitTextures.cpp
@@ -96,8 +96,6 @@ void UAppleARKitTextureCameraDepth::Init(float InTimestamp, AVDepthData* InCamer
 
 #endif
 
-//@joeg -- Added for environment capture support
-
 UAppleARKitEnvironmentCaptureProbeTexture::UAppleARKitEnvironmentCaptureProbeTexture(const FObjectInitializer& ObjectInitializer)
 	: Super(ObjectInitializer)
 #if PLATFORM_MAC || PLATFORM_IOS
diff --git a/Engine/Plugins/Runtime/AppleARKit/Source/AppleARKit/Public/AppleARKitAvailability.h b/Engine/Plugins/Runtime/AppleARKit/Source/AppleARKit/Public/AppleARKitAvailability.h
--- a/Engine/Plugins/Runtime/AppleARKit/Source/AppleARKit/Public/AppleARKitAvailability.h
+++ b/Engine/Plugins/Runtime/AppleARKit/Source/AppleARKit/Public/AppleARKitAvailability.h
@@ -3,7 +3,6 @@
 #pragma once
 
 
-//@joeg -- Fix for non-unity build error
 #include "CoreMinimal.h"
 #if PLATFORM_IOS && !PLATFORM_TVOS
 	#include "Availability.h"
@@ -25,10 +24,8 @@
 		#define SUPPORTS_ARKIT_1_5 0
 	#endif
 
-//@joeg -- Added image tracking support
 	// Check for ARKit 2.0
 	#ifdef __IPHONE_12_0
-//	#ifdef __IPHONE_11_3
 		#define SUPPORTS_ARKIT_2_0 1
 	#else
 		#define SUPPORTS_ARKIT_2_0 0
@@ -39,7 +36,6 @@
 	// No ARKit support
 	#define SUPPORTS_ARKIT_1_0 0
 	#define SUPPORTS_ARKIT_1_5 0
-//@joeg -- Added image tracking support
 	#define SUPPORTS_ARKIT_2_0 0
 
 #endif
@@ -83,7 +79,6 @@ public:
 		return bSupportsARKit15;
 	}
 
-//@joeg -- Added image tracking support
 	static bool SupportsARKit20()
 	{
 		static bool bSupportsARKit20 = false;
@@ -94,7 +89,6 @@ public:
 			bSupportChecked = true;
 			// This call is slow, so cache the result
 			if (@available(iOS 12.0, *))
-//			if (@available(iOS 11.3, *))
 			{
 				bSupportsARKit20 = true;
 			}
diff --git a/Engine/Plugins/Runtime/AppleARKit/Source/AppleARKit/Public/AppleARKitConversion.h b/Engine/Plugins/Runtime/AppleARKit/Source/AppleARKit/Public/AppleARKitConversion.h
--- a/Engine/Plugins/Runtime/AppleARKit/Source/AppleARKit/Public/AppleARKitConversion.h
+++ b/Engine/Plugins/Runtime/AppleARKit/Source/AppleARKit/Public/AppleARKitConversion.h
@@ -247,7 +247,6 @@ struct FAppleARKitConversion
 		{
 			WorldConfig.detectionImages = InitImageDetection(SessionConfig, CandidateImages, ConvertedCandidateImages);
 		}
-	//@joeg -- Added image tracking support
 	#if SUPPORTS_ARKIT_2_0
 		if (FAppleARKitAvailability::SupportsARKit20())
 		{
@@ -257,7 +256,6 @@ struct FAppleARKitConversion
 	}
 	#endif
 
-	//@joeg -- Added image tracking support
 	#if SUPPORTS_ARKIT_2_0
 	static void InitImageDetection(UARSessionConfig* SessionConfig, ARImageTrackingConfiguration* ImageConfig, TMap< FString, UARCandidateImage* >& CandidateImages, TMap< FString, CGImageRef >& ConvertedCandidateImages)
 	{
@@ -266,7 +264,6 @@ struct FAppleARKitConversion
 		ImageConfig.autoFocusEnabled = SessionConfig->ShouldEnableAutoFocus();
 	}
 
-	//@joeg -- Added environmental texture probe support
 	static AREnvironmentTexturing ToAREnvironmentTexturing(EAREnvironmentCaptureProbeType CaptureType)
 	{
 		switch (CaptureType)
@@ -289,7 +286,6 @@ struct FAppleARKitConversion
 		ARWorldMap* WorldMap = [NSKeyedUnarchiver unarchiveObjectWithData: WorldNSData];
 		if (WorldMap == nullptr)
 		{
-//@todo joeg -- Make a way to log from here
 //			UE_LOG(LogAppleARKit, Log, TEXT("Failed to load the world map data from the session object"));
 		}
 		return WorldMap;
@@ -321,13 +317,11 @@ struct FAppleARKitConversion
 				}
 				else
 				{
-//@todo joeg -- Make a way to log from here
 //					UE_LOG(LogAppleARKit, Log, TEXT("Failed to convert to ARReferenceObject (%s)"), *Candidate->GetFriendlyName());
 				}
 			}
 			else
 			{
-//@todo joeg -- Make a way to log from here
 //				UE_LOG(LogAppleARKit, Log, TEXT("Missing candidate object data for ARCandidateObject (%s)"), Candidate != nullptr ? *Candidate->GetFriendlyName() : TEXT("null"));
 			}
 		}
@@ -335,7 +329,6 @@ struct FAppleARKitConversion
 	}
 	#endif
 
-//@joeg -- Object detection addition
 	static ARConfiguration* ToARConfiguration( UARSessionConfig* SessionConfig, TMap< FString, UARCandidateImage* >& CandidateImages, TMap< FString, CGImageRef >& ConvertedCandidateImages, TMap< FString, UARCandidateObject* >& CandidateObjects )
 	{
 		EARSessionType SessionType = SessionConfig->GetSessionType();
@@ -380,7 +373,6 @@ struct FAppleARKitConversion
 					}
 				}
 #endif
-//@joeg -- Added environmental texture probe support
 #if SUPPORTS_ARKIT_2_0
 				if (FAppleARKitAvailability::SupportsARKit20())
 				{
@@ -400,7 +392,6 @@ struct FAppleARKitConversion
 				SessionConfiguration = WorldTrackingConfiguration;
 				break;
 			}
-//@joeg -- Added image tracking support
 			case EARSessionType::Image:
 			{
 #if SUPPORTS_ARKIT_2_0
@@ -418,7 +409,6 @@ struct FAppleARKitConversion
 #endif
 				break;
 			}
-//@joeg -- Added object scanning support
 			case EARSessionType::ObjectScanning:
 			{
 #if SUPPORTS_ARKIT_2_0
@@ -467,9 +457,7 @@ enum class EAppleAnchorType : uint8
 	PlaneAnchor,
 	FaceAnchor,
 	ImageAnchor,
-//@joeg -- Added environmental texture probe support
 	EnvironmentProbeAnchor,
-//@joeg -- Object detection
 	ObjectAnchor,
 	MAX
 };
@@ -493,35 +481,32 @@ struct FAppleARKitAnchorData
 	{
 	}
 
-//@joeg -- Eye tracking support
 	FAppleARKitAnchorData(FGuid InAnchorGuid, FTransform InTransform, FARBlendShapeMap InBlendShapes, TArray<FVector> InFaceVerts, FTransform InLeftEyeTransform, FTransform InRightEyeTransform, FVector InLookAtTarget)
-	: Transform( InTransform )
-	, AnchorType( EAppleAnchorType::FaceAnchor )
-	, AnchorGUID( InAnchorGuid )
-	, BlendShapes( MoveTemp(InBlendShapes) )
-	, FaceVerts( MoveTemp(InFaceVerts) )
-	, LeftEyeTransform( InLeftEyeTransform )
-	, RightEyeTransform( InRightEyeTransform )
-	, LookAtTarget( InLookAtTarget )
+		: Transform( InTransform )
+		, AnchorType( EAppleAnchorType::FaceAnchor )
+		, AnchorGUID( InAnchorGuid )
+		, BlendShapes( MoveTemp(InBlendShapes) )
+		, FaceVerts( MoveTemp(InFaceVerts) )
+		, LeftEyeTransform( InLeftEyeTransform )
+		, RightEyeTransform( InRightEyeTransform )
+		, LookAtTarget( InLookAtTarget )
 	{
 	}
 
-//@joeg -- Changed so object and image detection share
 	FAppleARKitAnchorData(FGuid InAnchorGuid, FTransform InTransform, EAppleAnchorType InAnchorType, FString InDetectedAnchorName)
-    : Transform( InTransform )
-    , AnchorType( InAnchorType )
-    , AnchorGUID( InAnchorGuid )
-    , DetectedAnchorName( MoveTemp(InDetectedAnchorName) )
+		: Transform( InTransform )
+		, AnchorType( InAnchorType )
+		, AnchorGUID( InAnchorGuid )
+		, DetectedAnchorName( MoveTemp(InDetectedAnchorName) )
     {
     }
 
-//@joeg -- Added environmental texture probe support
 	FAppleARKitAnchorData(FGuid InAnchorGuid, FTransform InTransform, FVector InExtent, id<MTLTexture> InProbeTexture)
-	: Transform( InTransform )
-	, AnchorType( EAppleAnchorType::EnvironmentProbeAnchor )
-	, AnchorGUID( InAnchorGuid )
-	, Extent(InExtent)
-	, ProbeTexture(InProbeTexture)
+		: Transform( InTransform )
+		, AnchorType( EAppleAnchorType::EnvironmentProbeAnchor )
+		, AnchorGUID( InAnchorGuid )
+		, Extent(InExtent)
+		, ProbeTexture(InProbeTexture)
 	{
 	}
 
@@ -537,13 +522,10 @@ struct FAppleARKitAnchorData
 	// Note: the index buffer never changes so can be safely read once
 	static TArray<int32> FaceIndices;
 
-//@joeg -- Changed this so both Image and Object anchors can share
 	FString DetectedAnchorName;
 
-//@joeg -- Environment texturing support
 	id<MTLTexture> ProbeTexture;
 
-//@joeg -- Eye tracking support
 	FTransform LeftEyeTransform;
 	FTransform RightEyeTransform;
 	FVector LookAtTarget;
diff --git a/Engine/Plugins/Runtime/AppleARKit/Source/AppleARKit/Public/AppleARKitFrame.h b/Engine/Plugins/Runtime/AppleARKit/Source/AppleARKit/Public/AppleARKitFrame.h
--- a/Engine/Plugins/Runtime/AppleARKit/Source/AppleARKit/Public/AppleARKitFrame.h
+++ b/Engine/Plugins/Runtime/AppleARKit/Source/AppleARKit/Public/AppleARKitFrame.h
@@ -79,7 +79,6 @@ struct APPLEARKIT_API FAppleARKitFrame
 	 */
 	FAppleARKitLightEstimate LightEstimate;
 	
-//@joeg -- ARKit 2.0 support
 	/** The current world mapping state is reported on the frame */
 	EARWorldMappingState WorldMappingState;
 
diff --git a/Engine/Plugins/Runtime/AppleARKit/Source/AppleARKit/Public/AppleARKitTextures.h b/Engine/Plugins/Runtime/AppleARKit/Source/AppleARKit/Public/AppleARKitTextures.h
--- a/Engine/Plugins/Runtime/AppleARKit/Source/AppleARKit/Public/AppleARKitTextures.h
+++ b/Engine/Plugins/Runtime/AppleARKit/Source/AppleARKit/Public/AppleARKitTextures.h
@@ -84,7 +84,6 @@ private:
 #endif
 };
 
-//@joeg -- Added for environment capture support
 UCLASS(BlueprintType)
 class APPLEARKIT_API UAppleARKitEnvironmentCaptureProbeTexture :
 	public UAREnvironmentCaptureProbeTexture,
diff --git a/Engine/Plugins/Runtime/AppleARKit/Source/AppleARKit/Public/AppleARKitTrackable.h b/Engine/Plugins/Runtime/AppleARKit/Source/AppleARKit/Public/AppleARKitTrackable.h
--- a/Engine/Plugins/Runtime/AppleARKit/Source/AppleARKit/Public/AppleARKitTrackable.h
+++ b/Engine/Plugins/Runtime/AppleARKit/Source/AppleARKit/Public/AppleARKitTrackable.h
@@ -8,7 +8,6 @@
 
 class UAppleARKitEnvironmentCaptureProbeTexture;
 
-//@joeg -- Added support for environment capture
 UCLASS(BlueprintType)
 class APPLEARKIT_API UAppleARKitEnvironmentCaptureProbe :
 	public UAREnvironmentCaptureProbe
diff --git a/Engine/Plugins/Runtime/AppleARKitFaceSupport/Source/AppleARKitFaceSupport/Private/AppleARKitFaceMeshConversion.h b/Engine/Plugins/Runtime/AppleARKitFaceSupport/Source/AppleARKitFaceSupport/Private/AppleARKitFaceMeshConversion.h
--- a/Engine/Plugins/Runtime/AppleARKitFaceSupport/Source/AppleARKitFaceSupport/Private/AppleARKitFaceMeshConversion.h
+++ b/Engine/Plugins/Runtime/AppleARKitFaceSupport/Source/AppleARKitFaceSupport/Private/AppleARKitFaceMeshConversion.h
@@ -46,7 +46,6 @@ static FORCEINLINE FARBlendShapeMap ToBlendShapeMap(NSDictionary<ARBlendShapeLoc
 	BlendShapeMap.Add(EARFaceBlendShape::HeadYaw, HeadYaw);
 	BlendShapeMap.Add(EARFaceBlendShape::HeadPitch, HeadPitch);
 	BlendShapeMap.Add(EARFaceBlendShape::HeadRoll, HeadRoll);
-//@joeg -- Eye tracking support
 	FRotator LeftEyeRot(LeftEyeTransform.GetRotation());
 	float LeftEyeYaw = (float)LeftEyeRot.Yaw / 180.f;
 	float LeftEyePitch = (float)LeftEyeRot.Pitch / 180.f;
@@ -174,7 +173,6 @@ static FORCEINLINE FARBlendShapeMap ToBlendShapeMap(NSDictionary<ARBlendShapeLoc
 	SET_BLEND_SHAPE(ARBlendShapeLocationMouthShrugUpper, EARFaceBlendShape::MouthShrugUpper);
 	SET_BLEND_SHAPE(ARBlendShapeLocationBrowInnerUp, EARFaceBlendShape::BrowInnerUp);
 	SET_BLEND_SHAPE(ARBlendShapeLocationCheekPuff, EARFaceBlendShape::CheekPuff);
-//@joeg -- Tongue blend shape
 #if SUPPORTS_ARKIT_2_0
 	if (FAppleARKitAvailability::SupportsARKit20())
 	{
diff --git a/Engine/Plugins/Runtime/AppleARKitFaceSupport/Source/AppleARKitFaceSupport/Private/AppleARKitFaceSupportImpl.cpp b/Engine/Plugins/Runtime/AppleARKitFaceSupport/Source/AppleARKitFaceSupport/Private/AppleARKitFaceSupportImpl.cpp
--- a/Engine/Plugins/Runtime/AppleARKitFaceSupport/Source/AppleARKitFaceSupport/Private/AppleARKitFaceSupportImpl.cpp
+++ b/Engine/Plugins/Runtime/AppleARKitFaceSupport/Source/AppleARKitFaceSupport/Private/AppleARKitFaceSupportImpl.cpp
@@ -21,7 +21,6 @@ static TSharedPtr<FAppleARKitAnchorData> MakeAnchorData(ARAnchor* Anchor, const
     if ([Anchor isKindOfClass:[ARFaceAnchor class]])
     {
         ARFaceAnchor* FaceAnchor = (ARFaceAnchor*)Anchor;
-//@joeg -- Eye tracking support
 		FTransform LeftEyeTransform;
 		FTransform RightEyeTransform;
 		FVector LookAtTarget;
diff --git a/Engine/Source/Runtime/AugmentedReality/Private/ARBlueprintLibrary.cpp b/Engine/Source/Runtime/AugmentedReality/Private/ARBlueprintLibrary.cpp
--- a/Engine/Source/Runtime/AugmentedReality/Private/ARBlueprintLibrary.cpp
+++ b/Engine/Source/Runtime/AugmentedReality/Private/ARBlueprintLibrary.cpp
@@ -340,8 +340,6 @@ UARTextureCameraDepth* UARBlueprintLibrary::GetCameraDepth()
 	return Depth;
 }
 
-//@joeg -- Helpers to make things easier
-
 TArray<UARPlaneGeometry*> UARBlueprintLibrary::GetAllTrackedPlanes()
 {
 	TArray<UARPlaneGeometry*> ResultList;
@@ -414,7 +412,6 @@ TArray<UAREnvironmentCaptureProbe*> UARBlueprintLibrary::GetAllTrackedEnvironmen
 	return ResultList;
 }
 
-//@joeg -- Added environmental texture probe support
 bool UARBlueprintLibrary::AddManualEnvironmentCaptureProbe(FVector Location, FVector Extent)
 {
 	auto ARSystem = GetARSystem();
@@ -457,8 +454,6 @@ TSharedPtr<FARGetCandidateObjectAsyncTask, ESPMode::ThreadSafe> UARBlueprintLibr
 	return TSharedPtr<FARGetCandidateObjectAsyncTask, ESPMode::ThreadSafe>();
 }
 
-//@joeg -- End additions
-
 TArray<FARVideoFormat> UARBlueprintLibrary::GetSupportedVideoFormats(EARSessionType SessionType)
 {
 	auto ARSystem = GetARSystem();
diff --git a/Engine/Source/Runtime/AugmentedReality/Private/ARBlueprintProxy.cpp b/Engine/Source/Runtime/AugmentedReality/Private/ARBlueprintProxy.cpp
--- a/Engine/Source/Runtime/AugmentedReality/Private/ARBlueprintProxy.cpp
+++ b/Engine/Source/Runtime/AugmentedReality/Private/ARBlueprintProxy.cpp
@@ -1,5 +1,3 @@
-//@joeg -- ARkit 2.0 additions
-
 // Copyright 1998-2018 Epic Games, Inc. All Rights Reserved.
 
 #include "ARBlueprintProxy.h"
diff --git a/Engine/Source/Runtime/AugmentedReality/Private/ARSessionConfig.cpp b/Engine/Source/Runtime/AugmentedReality/Private/ARSessionConfig.cpp
--- a/Engine/Source/Runtime/AugmentedReality/Private/ARSessionConfig.cpp
+++ b/Engine/Source/Runtime/AugmentedReality/Private/ARSessionConfig.cpp
@@ -14,7 +14,6 @@ UARSessionConfig::UARSessionConfig()
 , FrameSyncMode(EARFrameSyncMode::SyncTickWithoutCameraImage)
 , bEnableAutomaticCameraOverlay(true)
 , bEnableAutomaticCameraTracking(true)
-//@joeg -- Added image tracking support
 , MaxNumSimultaneousImagesTracked(1)
 {
 }
@@ -66,20 +65,17 @@ const TArray<UARCandidateImage*>& UARSessionConfig::GetCandidateImageList() cons
 	return CandidateImages;
 }
 
-//@joeg -- Added image tracking support
 int32 UARSessionConfig::GetMaxNumSimultaneousImagesTracked() const
 {
     return MaxNumSimultaneousImagesTracked;
 }
 
-//@joeg -- Added environmental texture probe support
 /** @see EnvironmentCaptureProbeType */
 EAREnvironmentCaptureProbeType UARSessionConfig::GetEnvironmentCaptureProbeType() const
 {
 	return EnvironmentCaptureProbeType;
 }
 
-//@joeg -- Load/Save world support
 const TArray<uint8>& UARSessionConfig::GetWorldMapData() const
 {
 	return WorldMapData;
@@ -89,9 +85,7 @@ void UARSessionConfig::SetWorldMapData(TArray<uint8> InWorldMapData)
 {
 	WorldMapData = MoveTemp(InWorldMapData);
 }
-//@joeg -- end load/save
 
-//@joeg -- Object detection
 const TArray<UARCandidateObject*>& UARSessionConfig::GetCandidateObjectList() const
 {
 	return CandidateObjects;
@@ -104,7 +98,6 @@ void UARSessionConfig::AddCandidateObject(UARCandidateObject* CandidateObject)
 		CandidateObjects.Add(CandidateObject);
 	}
 }
-//@joeg -- end object detection
 
 FARVideoFormat UARSessionConfig::GetDesiredVideoFormat() const
 {
diff --git a/Engine/Source/Runtime/AugmentedReality/Private/ARSystem.cpp b/Engine/Source/Runtime/AugmentedReality/Private/ARSystem.cpp
--- a/Engine/Source/Runtime/AugmentedReality/Private/ARSystem.cpp
+++ b/Engine/Source/Runtime/AugmentedReality/Private/ARSystem.cpp
@@ -97,8 +97,6 @@ UARTextureCameraDepth* FARSystemBase::GetCameraDepth()
 	return OnGetCameraDepth();
 }
 
-//@joeg -- ARKit 2.0 additions
-
 bool FARSystemBase::AddManualEnvironmentCaptureProbe(FVector Location, FVector Extent)
 {
 	return OnAddManualEnvironmentCaptureProbe(Location, Extent);
@@ -118,7 +116,6 @@ EARWorldMappingState FARSystemBase::GetWorldMappingStatus() const
 {
 	return OnGetWorldMappingStatus();
 }
-//@joeg -- End additions
 
 bool FARSystemBase::IsSessionTypeSupported(EARSessionType SessionType) const
 {
diff --git a/Engine/Source/Runtime/AugmentedReality/Private/ARTextures.cpp b/Engine/Source/Runtime/AugmentedReality/Private/ARTextures.cpp
--- a/Engine/Source/Runtime/AugmentedReality/Private/ARTextures.cpp
+++ b/Engine/Source/Runtime/AugmentedReality/Private/ARTextures.cpp
@@ -17,7 +17,6 @@ UARTextureCameraDepth::UARTextureCameraDepth(const FObjectInitializer& ObjectIni
 {
 }
 
-//@joeg -- Added environmental texture probe support
 UAREnvironmentCaptureProbeTexture::UAREnvironmentCaptureProbeTexture(const FObjectInitializer& ObjectInitializer)
 	: Super(ObjectInitializer)
 	, TextureType(EARTextureType::EnvironmentCapture)
diff --git a/Engine/Source/Runtime/AugmentedReality/Private/ARTrackable.cpp b/Engine/Source/Runtime/AugmentedReality/Private/ARTrackable.cpp
--- a/Engine/Source/Runtime/AugmentedReality/Private/ARTrackable.cpp
+++ b/Engine/Source/Runtime/AugmentedReality/Private/ARTrackable.cpp
@@ -184,7 +184,6 @@ void UARFaceGeometry::UpdateFaceGeometry(const TSharedRef<FARSystemBase, ESPMode
 		IndexBuffer = Indices;
 	}
 	
-//@joeg -- Eye tracking support
 	LeftEyeTransform = InLeftEyeTransform;
 	RightEyeTransform = InRightEyeTransform;
 	LookAtTarget = InLookAtTarget;
@@ -225,7 +224,6 @@ const TMap<EARFaceBlendShape, float> UARFaceGeometry::GetBlendShapes() const
 	return BlendShapes;
 }
 
-//@joeg -- Eye tracking support
 const FTransform& UARFaceGeometry::GetLocalSpaceEyeTransform(EAREye Eye) const
 {
 	if (Eye == EAREye::LeftEye)
@@ -243,10 +241,7 @@ FTransform UARFaceGeometry::GetWorldSpaceEyeTransform(EAREye Eye) const
 	}
 	return GetLocalToWorldTransform() * RightEyeTransform;
 }
-//@joeg -- end eye tracking
 
-
-//@joeg -- Added environmental texture probe support
 UAREnvironmentCaptureProbe::UAREnvironmentCaptureProbe()
 	: Super()
 	, EnvironmentCaptureTexture(nullptr)
@@ -281,7 +276,6 @@ UAREnvironmentCaptureProbeTexture* UAREnvironmentCaptureProbe::GetEnvironmentCap
 	return EnvironmentCaptureTexture;
 }
 
-//@joeg -- Object detection
 void UARTrackedObject::DebugDraw(UWorld* World, const FLinearColor& OutlineColor, float OutlineThickness, float PersistForSeconds /*= 0.0f*/) const
 {
 	const FTransform LocalToWorldTransform = GetLocalToWorldTransform();
diff --git a/Engine/Source/Runtime/AugmentedReality/Private/ARTypes.cpp b/Engine/Source/Runtime/AugmentedReality/Private/ARTypes.cpp
--- a/Engine/Source/Runtime/AugmentedReality/Private/ARTypes.cpp
+++ b/Engine/Source/Runtime/AugmentedReality/Private/ARTypes.cpp
@@ -4,8 +4,6 @@
 #include "ARTypes.h"
 #include "HAL/ThreadSafeBool.h"
 
-//@joeg -- ARKit 2.0 support
-
 bool FARAsyncTask::HadError() const
 {
 	return bHadError;
@@ -33,5 +31,3 @@ TArray<uint8> FARSaveWorldAsyncTask::GetSavedWorldData()
 	}
 	return TArray<uint8>();
 }
-
-//@joeg -- End additions
diff --git a/Engine/Source/Runtime/AugmentedReality/Public/ARBlueprintLibrary.h b/Engine/Source/Runtime/AugmentedReality/Public/ARBlueprintLibrary.h
--- a/Engine/Source/Runtime/AugmentedReality/Public/ARBlueprintLibrary.h
+++ b/Engine/Source/Runtime/AugmentedReality/Public/ARBlueprintLibrary.h
@@ -158,7 +158,6 @@ public:
 	UFUNCTION(BlueprintCallable, Category = "AR AugmentedReality|Pin", meta = (Keywords = "ar augmentedreality augmented reality tracking pin anchor"))
 	static TArray<UARPin*> GetAllPins();
 	
-//@joeg -- Helpers to make things easier
 	/** @return a list of all the tracked planes as currently seen by the Augmented Reality system */
 	UFUNCTION(BlueprintCallable, Category = "AR AugmentedReality|Tracking", meta = (Keywords = "ar augmentedreality augmented reality tracking geometry anchor"))
 	static TArray<UARPlaneGeometry*> GetAllTrackedPlanes();
@@ -175,7 +174,6 @@ public:
 	UFUNCTION(BlueprintCallable, Category = "AR AugmentedReality|Tracking", meta = (Keywords = "ar augmentedreality augmented reality tracking anchor"))
 	static TArray<UAREnvironmentCaptureProbe*> GetAllTrackedEnvironmentCaptureProbes();
 
-//@joeg -- Added environmental texture probe support
 	/** Adds an environment capture probe to the ar world */
 	UFUNCTION(BlueprintCallable, Category = "AR AugmentedReality|Tracking", meta = (Keywords = "ar augmentedreality augmented reality tracking anchor"))
 	static bool AddManualEnvironmentCaptureProbe(FVector Location, FVector Extent);
@@ -191,8 +189,6 @@ public:
 	static TSharedPtr<FARSaveWorldAsyncTask, ESPMode::ThreadSafe> SaveWorld();
 	static TSharedPtr<FARGetCandidateObjectAsyncTask, ESPMode::ThreadSafe> GetCandidateObject(FVector Location, FVector Extent);
 	
-//@joeg -- End additions
-
 public:
 	static void RegisterAsARSystem(const TSharedPtr<FARSystemBase, ESPMode::ThreadSafe>& NewArSystem);
 	
diff --git a/Engine/Source/Runtime/AugmentedReality/Public/ARBlueprintProxy.h b/Engine/Source/Runtime/AugmentedReality/Public/ARBlueprintProxy.h
--- a/Engine/Source/Runtime/AugmentedReality/Public/ARBlueprintProxy.h
+++ b/Engine/Source/Runtime/AugmentedReality/Public/ARBlueprintProxy.h
@@ -1,5 +1,3 @@
-//@joeg -- ARkit 2.0 additions
-
 // Copyright 1998-2018 Epic Games, Inc. All Rights Reserved.
 
 #pragma once
diff --git a/Engine/Source/Runtime/AugmentedReality/Public/ARSessionConfig.h b/Engine/Source/Runtime/AugmentedReality/Public/ARSessionConfig.h
--- a/Engine/Source/Runtime/AugmentedReality/Public/ARSessionConfig.h
+++ b/Engine/Source/Runtime/AugmentedReality/Public/ARSessionConfig.h
@@ -40,11 +40,9 @@ enum class EARSessionType : uint8
 	/** AR meant to overlay onto a face */
 	Face,
 
-//@joeg -- Added image tracking support
     /** Tracking of images supplied by the app. No world tracking, just images */
     Image,
 
-//@joeg -- Object scanning support
 	/** A session used to scan objects for object detection in a world tracking session */
 	ObjectScanning
 };
@@ -85,7 +83,6 @@ enum class EARFrameSyncMode : uint8
 	SyncTickWithoutCameraImage = 1,
 };
 
-//@joeg -- Added texture probe settings
 /**
  * Tells the AR system what type of environmental texture capturing to perform
  */
@@ -137,21 +134,17 @@ public:
 	/** @see CandidateImages */
 	const TArray<UARCandidateImage*>& GetCandidateImageList() const;
     
-//@joeg -- Added image tracking support
 	/** @see MaxNumSimultaneousImagesTracked */
     int32 GetMaxNumSimultaneousImagesTracked() const;
 	
-//@joeg -- Added environmental texture probe support
 	/** @see EnvironmentCaptureProbeType */
 	EAREnvironmentCaptureProbeType GetEnvironmentCaptureProbeType() const;
 	
-//@joeg -- Added for load/save of worlds
 	/** @see WorldMapData */
 	const TArray<uint8>& GetWorldMapData() const;
 	/** @see WorldMapData */
 	void SetWorldMapData(TArray<uint8> WorldMapData);
 
-//@joeg -- For object detection
 	const TArray<UARCandidateObject*>& GetCandidateObjectList() const;
 	void AddCandidateObject(UARCandidateObject* CandidateObject);
 	
@@ -209,22 +202,18 @@ protected:
 	UPROPERTY(EditAnywhere, Category="AR Settings")
 	TArray<UARCandidateImage*> CandidateImages;
 
-//@joeg -- Added image tracking support
     /** The maximum number of images to track at the same time. Defaults to 1 */
     UPROPERTY(EditAnywhere, Category="AR Settings")
     int32 MaxNumSimultaneousImagesTracked;
 	
-//@joeg -- Added environmental texture probe support
 	/** How the AR system should handle texture probe capturing */
 	UPROPERTY(EditAnywhere, Category="AR Settings")
 	EAREnvironmentCaptureProbeType EnvironmentCaptureProbeType;
 
-//@joeg -- For loading a saved world
 	/** A previously saved world that is to be loaded when the session starts */
 	UPROPERTY(VisibleAnywhere, Category="AR Settings")
 	TArray<uint8> WorldMapData;
 
-//@joeg -- For object detection
 	/** A list of candidate objects to search for in the scene */
 	UPROPERTY(EditAnywhere, Category="AR Settings")
 	TArray<UARCandidateObject*> CandidateObjects;
diff --git a/Engine/Source/Runtime/AugmentedReality/Public/ARSystem.h b/Engine/Source/Runtime/AugmentedReality/Public/ARSystem.h
--- a/Engine/Source/Runtime/AugmentedReality/Public/ARSystem.h
+++ b/Engine/Source/Runtime/AugmentedReality/Public/ARSystem.h
@@ -120,7 +120,6 @@ public:
 	/** @return the last camera depth information the AR system has seen */
 	virtual UARTextureCameraDepth* OnGetCameraDepth() = 0;
 
-//@joeg -- ARKit 2.0 additions
 	/** Tells the ARSystem to generate a capture probe at the specified location if supported */
 	virtual bool OnAddManualEnvironmentCaptureProbe(FVector Location, FVector Extent) = 0;
 	
@@ -132,7 +131,6 @@ public:
 	
 	/** @return the current mapping status */
 	virtual EARWorldMappingState OnGetWorldMappingStatus() const = 0;
-//@joeg -- End additions
 	
 	/** @return The list of supported video formats for this device and session type */
 	virtual TArray<FARVideoFormat> OnGetSupportedVideoFormats(EARSessionType SessionType) const = 0;
@@ -196,7 +194,6 @@ public:
 	UARTextureCameraImage* GetCameraImage();
 	/** \see UARBlueprintLibrary::GetCameraDepth() */
 	UARTextureCameraDepth* GetCameraDepth();
-//@joeg -- ARKit 2.0 additions
 	/**\see UARBlueprintLibrary::IsEnvironmentCaptureSupported() */
 	bool IsEnvironmentCaptureSupported() const;
 	/**\see UARBlueprintLibrary::AddEnvironmentCaptureProbe() */
@@ -207,7 +204,6 @@ public:
 	TSharedPtr<FARSaveWorldAsyncTask, ESPMode::ThreadSafe> SaveWorld() const;
 	/** @return the current mapping status */
 	EARWorldMappingState GetWorldMappingStatus() const;
-//@joeg -- End additions
 
 	/** \see UARBlueprintLibrary::GetCurrentLightEstimate() */
 	UARLightEstimate* GetCurrentLightEstimate() const;
diff --git a/Engine/Source/Runtime/AugmentedReality/Public/ARTextures.h b/Engine/Source/Runtime/AugmentedReality/Public/ARTextures.h
--- a/Engine/Source/Runtime/AugmentedReality/Public/ARTextures.h
+++ b/Engine/Source/Runtime/AugmentedReality/Public/ARTextures.h
@@ -13,7 +13,6 @@ enum class EARTextureType : uint8
 {
 	CameraImage,
 	CameraDepth,
-//@joeg -- Addedd environment capture support
 	EnvironmentCapture
 };
 
@@ -97,7 +96,6 @@ public:
 	bool bIsTemporallySmoothed;
 };
 
-//@joeg -- Added for environment capture support
 /**
  * Base class for all AR textures that represent the environment for lighting and reflection
  */
diff --git a/Engine/Source/Runtime/AugmentedReality/Public/ARTrackable.h b/Engine/Source/Runtime/AugmentedReality/Public/ARTrackable.h
--- a/Engine/Source/Runtime/AugmentedReality/Public/ARTrackable.h
+++ b/Engine/Source/Runtime/AugmentedReality/Public/ARTrackable.h
@@ -223,13 +223,11 @@ enum class EARFaceBlendShape : uint8
 	// Nose blend shapes
 	NoseSneerLeft,
 	NoseSneerRight,
-//@joeg -- Tongue blend shape support
 	TongueOut,
 	// Treat the head rotation as curves for LiveLink support
 	HeadYaw,
 	HeadPitch,
 	HeadRoll,
-//@joeg -- Eye tracking support
 	// Treat eye rotation as curves for LiveLink support
 	LeftEyeYaw,
 	LeftEyePitch,
@@ -240,7 +238,6 @@ enum class EARFaceBlendShape : uint8
 	MAX
 };
 
-//@joeg -- Eye tracking support
 UENUM(BlueprintType, Category="AR AugmentedReality", meta=(Experimental))
 enum class EAREye : uint8
 {
@@ -297,14 +294,12 @@ private:
 	// @todo JoeG - route the uvs in
 	TArray<FVector2D> UVs;
 
-//@joeg -- Eye tracking support
 	/** The transform for the left eye */
 	FTransform LeftEyeTransform;
 	/** The transform for the right eye */
 	FTransform RightEyeTransform;
 };
 
-//@joeg -- Added environmental texture probe support
 /** A tracked environment texture probe that gives you a cube map for reflections */
 UCLASS(BlueprintType)
 class AUGMENTEDREALITY_API UAREnvironmentCaptureProbe :
@@ -336,7 +331,6 @@ protected:
 	UAREnvironmentCaptureProbeTexture* EnvironmentCaptureTexture;
 };
 
-//@joeg -- Added for object detection
 UCLASS(BlueprintType)
 class AUGMENTEDREALITY_API UARTrackedObject : public UARTrackedGeometry
 {
diff --git a/Engine/Source/Runtime/AugmentedReality/Public/ARTypes.h b/Engine/Source/Runtime/AugmentedReality/Public/ARTypes.h
--- a/Engine/Source/Runtime/AugmentedReality/Public/ARTypes.h
+++ b/Engine/Source/Runtime/AugmentedReality/Public/ARTypes.h
@@ -3,7 +3,6 @@
 #pragma once
 
 #include "Engine/DataAsset.h"
-//@joeg -- ARKit 2.0 support
 #include "HAL/ThreadSafeBool.h"
 #include "ARTypes.generated.h"
 
@@ -90,7 +89,6 @@ enum class EARSessionStatus : uint8
 	Other,
 };
 
-//@joeg -- ARKit 2.0 support
 /** Gives feedback as to whether the AR data can be saved and relocalized or not */
 UENUM(BlueprintType, Category="AR AugmentedReality", meta=(Experimental))
 enum class EARWorldMappingState : uint8
@@ -212,8 +210,6 @@ private:
 	EARCandidateImageOrientation Orientation;
 };
 
-//@joeg -- ARKit 2.0 support
-
 /** An asset that points to an object to be detected in a scene */
 UCLASS(BlueprintType)
 class AUGMENTEDREALITY_API UARCandidateObject :
