Changed how timestamps for ARKit objects are updated to make them more amenable with the engine

diff --git a/Engine/Plugins/Runtime/AppleARKit/Source/AppleARKit/Private/AppleARKitSystem.cpp b/Engine/Plugins/Runtime/AppleARKit/Source/AppleARKit/Private/AppleARKitSystem.cpp
--- a/Engine/Plugins/Runtime/AppleARKit/Source/AppleARKit/Private/AppleARKitSystem.cpp
+++ b/Engine/Plugins/Runtime/AppleARKit/Source/AppleARKit/Private/AppleARKitSystem.cpp
@@ -688,8 +688,10 @@ UARPin* FAppleARKitSystem::OnPinComponent( USceneComponent* ComponentToPin, cons
 		UARTrackedGeometry* GeometryToPinTo = TrackedGeometry;
 		if (GeometryToPinTo == nullptr)
 		{
+			double UpdateTimestamp = FPlatformTime::Seconds();
+			
 			GeometryToPinTo = NewObject<UARTrackedPoint>();
-			GeometryToPinTo->UpdateTrackedGeometry(SharedThis(this), GameThreadFrameNumber, GameThreadTimestamp, PinToTrackingTransform, GetAlignmentTransform());
+			GeometryToPinTo->UpdateTrackedGeometry(SharedThis(this), GameThreadFrameNumber, UpdateTimestamp, PinToTrackingTransform, GetAlignmentTransform());
 		}
 		
 		UARPin* NewPin = NewObject<UARPin>();
@@ -1348,9 +1350,11 @@ void FAppleARKitSystem::SessionDidAddAnchors_DelegateThread( NSArray<ARAnchor*>*
 	// If this object is valid, we are running a face session and need that code to process things
 	if (FaceARSupport != nullptr)
 	{
+		double UpdateTimestamp = FPlatformTime::Seconds();
 		const FRotator& AdjustBy = GetSessionConfig().GetWorldAlignment() == EARWorldAlignment::Camera ? DerivedTrackingToUnrealRotation : FRotator::ZeroRotator;
 		const EARFaceTrackingUpdate UpdateSetting = GetSessionConfig().GetFaceTrackingUpdate();
-		const TArray<TSharedPtr<FAppleARKitAnchorData>> AnchorList = FaceARSupport->MakeAnchorData(anchors, GameThreadTimestamp, GameThreadFrameNumber, AdjustBy, UpdateSetting);
+
+		const TArray<TSharedPtr<FAppleARKitAnchorData>> AnchorList = FaceARSupport->MakeAnchorData(anchors, UpdateTimestamp, GameThreadFrameNumber, AdjustBy, UpdateSetting);
 		for (TSharedPtr<FAppleARKitAnchorData> NewAnchorData : AnchorList)
 		{
 			auto AddAnchorTask = FSimpleDelegateGraphTask::FDelegate::CreateSP(this, &FAppleARKitSystem::SessionDidAddAnchors_Internal, NewAnchorData.ToSharedRef());
@@ -1379,9 +1383,11 @@ void FAppleARKitSystem::SessionDidUpdateAnchors_DelegateThread( NSArray<ARAnchor
 	// If this object is valid, we are running a face session and need that code to process things
 	if (FaceARSupport != nullptr)
 	{
+		double UpdateTimestamp = FPlatformTime::Seconds();
 		const FRotator& AdjustBy = GetSessionConfig().GetWorldAlignment() == EARWorldAlignment::Camera ? DerivedTrackingToUnrealRotation : FRotator::ZeroRotator;
 		const EARFaceTrackingUpdate UpdateSetting = GetSessionConfig().GetFaceTrackingUpdate();
-		const TArray<TSharedPtr<FAppleARKitAnchorData>> AnchorList = FaceARSupport->MakeAnchorData(anchors, GameThreadTimestamp, GameThreadFrameNumber, AdjustBy, UpdateSetting);
+
+		const TArray<TSharedPtr<FAppleARKitAnchorData>> AnchorList = FaceARSupport->MakeAnchorData(anchors, UpdateTimestamp, GameThreadFrameNumber, AdjustBy, UpdateSetting);
 		for (TSharedPtr<FAppleARKitAnchorData> NewAnchorData : AnchorList)
 		{
 			auto UpdateAnchorTask = FSimpleDelegateGraphTask::FDelegate::CreateSP(this, &FAppleARKitSystem::SessionDidUpdateAnchors_Internal, NewAnchorData.ToSharedRef());
@@ -1422,6 +1428,8 @@ void FAppleARKitSystem::SessionDidRemoveAnchors_DelegateThread( NSArray<ARAnchor
 
 void FAppleARKitSystem::SessionDidAddAnchors_Internal( TSharedRef<FAppleARKitAnchorData> AnchorData )
 {
+	double UpdateTimestamp = FPlatformTime::Seconds();
+	
 	// In case we have camera tracking turned off, we still need to update the frame
 	if (!GetSessionConfig().ShouldEnableCameraTracking())
 	{
@@ -1431,7 +1439,7 @@ void FAppleARKitSystem::SessionDidAddAnchors_Internal( TSharedRef<FAppleARKitAnc
 	// If this object is valid, we are running a face session and we need to publish LiveLink data on the game thread
 	if (FaceARSupport != nullptr && AnchorData->AnchorType == EAppleAnchorType::FaceAnchor)
 	{
-		FaceARSupport->PublishLiveLinkData(AnchorData, GameThreadTimestamp, GameThreadFrameNumber);
+		FaceARSupport->PublishLiveLinkData(AnchorData, UpdateTimestamp, GameThreadFrameNumber);
 	}
 
 	FString NewAnchorDebugName;
@@ -1442,14 +1450,14 @@ void FAppleARKitSystem::SessionDidAddAnchors_Internal( TSharedRef<FAppleARKitAnc
 		{
 			NewAnchorDebugName = FString::Printf(TEXT("ANCHOR-%02d"), LastTrackedGeometry_DebugId++);
 			NewGeometry = NewObject<UARTrackedGeometry>();
-			NewGeometry->UpdateTrackedGeometry(SharedThis(this), GameThreadFrameNumber, GameThreadTimestamp, AnchorData->Transform, GetAlignmentTransform());
+			NewGeometry->UpdateTrackedGeometry(SharedThis(this), GameThreadFrameNumber, UpdateTimestamp, AnchorData->Transform, GetAlignmentTransform());
 			break;
 		}
 		case EAppleAnchorType::PlaneAnchor:
 		{
 			NewAnchorDebugName = FString::Printf(TEXT("PLN-%02d"), LastTrackedGeometry_DebugId++);
 			UARPlaneGeometry* NewGeo = NewObject<UARPlaneGeometry>();
-			NewGeo->UpdateTrackedGeometry(SharedThis(this), GameThreadFrameNumber, GameThreadTimestamp, AnchorData->Transform, GetAlignmentTransform(), AnchorData->Center, AnchorData->Extent);
+			NewGeo->UpdateTrackedGeometry(SharedThis(this), GameThreadFrameNumber, UpdateTimestamp, AnchorData->Transform, GetAlignmentTransform(), AnchorData->Center, AnchorData->Extent);
 			NewGeometry = NewGeo;
 			break;
 		}
@@ -1457,7 +1465,7 @@ void FAppleARKitSystem::SessionDidAddAnchors_Internal( TSharedRef<FAppleARKitAnc
 		{
 			NewAnchorDebugName = FString::Printf(TEXT("FACE-%02d"), LastTrackedGeometry_DebugId++);
 			UARFaceGeometry* NewGeo = NewObject<UARFaceGeometry>();
-			NewGeo->UpdateFaceGeometry(SharedThis(this), GameThreadFrameNumber, GameThreadTimestamp, AnchorData->Transform, GetAlignmentTransform(), AnchorData->BlendShapes, AnchorData->FaceVerts, AnchorData->FaceIndices, AnchorData->LeftEyeTransform, AnchorData->RightEyeTransform, AnchorData->LookAtTarget);
+			NewGeo->UpdateFaceGeometry(SharedThis(this), GameThreadFrameNumber, UpdateTimestamp, AnchorData->Transform, GetAlignmentTransform(), AnchorData->BlendShapes, AnchorData->FaceVerts, AnchorData->FaceIndices, AnchorData->LeftEyeTransform, AnchorData->RightEyeTransform, AnchorData->LookAtTarget);
 			NewGeo->bIsTracked = true;
 			NewGeometry = NewGeo;
 			break;
@@ -1468,7 +1476,7 @@ void FAppleARKitSystem::SessionDidAddAnchors_Internal( TSharedRef<FAppleARKitAnc
 			UARTrackedImage* NewImage = NewObject<UARTrackedImage>();
 			UARCandidateImage** CandidateImage = CandidateImages.Find(AnchorData->DetectedAnchorName);
 			ensure(CandidateImage != nullptr);
-			NewImage->UpdateTrackedGeometry(SharedThis(this), GameThreadFrameNumber, GameThreadTimestamp, AnchorData->Transform, GetAlignmentTransform(), *CandidateImage);
+			NewImage->UpdateTrackedGeometry(SharedThis(this), GameThreadFrameNumber, UpdateTimestamp, AnchorData->Transform, GetAlignmentTransform(), *CandidateImage);
 			NewGeometry = NewImage;
 			break;
 		}
@@ -1476,7 +1484,7 @@ void FAppleARKitSystem::SessionDidAddAnchors_Internal( TSharedRef<FAppleARKitAnc
 		{
 			NewAnchorDebugName = FString::Printf(TEXT("ENV-%02d"), LastTrackedGeometry_DebugId++);
 			UAppleARKitEnvironmentCaptureProbe* NewProbe = NewObject<UAppleARKitEnvironmentCaptureProbe>();
-			NewProbe->UpdateEnvironmentCapture(SharedThis(this), GameThreadFrameNumber, GameThreadTimestamp, AnchorData->Transform, GetAlignmentTransform(), AnchorData->Extent, AnchorData->ProbeTexture);
+			NewProbe->UpdateEnvironmentCapture(SharedThis(this), GameThreadFrameNumber, UpdateTimestamp, AnchorData->Transform, GetAlignmentTransform(), AnchorData->Extent, AnchorData->ProbeTexture);
 			NewGeometry = NewProbe;
 			break;
 		}
@@ -1486,7 +1494,7 @@ void FAppleARKitSystem::SessionDidAddAnchors_Internal( TSharedRef<FAppleARKitAnc
 			UARTrackedObject* NewTrackedObject = NewObject<UARTrackedObject>();
 			UARCandidateObject** CandidateObject = CandidateObjects.Find(AnchorData->DetectedAnchorName);
 			ensure(CandidateObject != nullptr);
-			NewTrackedObject->UpdateTrackedGeometry(SharedThis(this), GameThreadFrameNumber, GameThreadTimestamp, AnchorData->Transform, GetAlignmentTransform(), *CandidateObject);
+			NewTrackedObject->UpdateTrackedGeometry(SharedThis(this), GameThreadFrameNumber, UpdateTimestamp, AnchorData->Transform, GetAlignmentTransform(), *CandidateObject);
 			NewGeometry = NewTrackedObject;
 			break;
 		}
@@ -1500,6 +1508,8 @@ void FAppleARKitSystem::SessionDidAddAnchors_Internal( TSharedRef<FAppleARKitAnc
 
 void FAppleARKitSystem::SessionDidUpdateAnchors_Internal( TSharedRef<FAppleARKitAnchorData> AnchorData )
 {
+	double UpdateTimestamp = FPlatformTime::Seconds();
+	
 	// In case we have camera tracking turned off, we still need to update the frame
 	if (!GetSessionConfig().ShouldEnableCameraTracking())
 	{
@@ -1509,7 +1519,7 @@ void FAppleARKitSystem::SessionDidUpdateAnchors_Internal( TSharedRef<FAppleARKit
 	// If this object is valid, we are running a face session and we need to publish LiveLink data on the game thread
 	if (FaceARSupport != nullptr && AnchorData->AnchorType == EAppleAnchorType::FaceAnchor)
 	{
-		FaceARSupport->PublishLiveLinkData(AnchorData, GameThreadTimestamp, GameThreadFrameNumber);
+		FaceARSupport->PublishLiveLinkData(AnchorData, UpdateTimestamp, GameThreadFrameNumber);
 	}
 
 	UARTrackedGeometry** GeometrySearchResult = TrackedGeometries.Find(AnchorData->AnchorGUID);
@@ -1530,7 +1540,7 @@ void FAppleARKitSystem::SessionDidUpdateAnchors_Internal( TSharedRef<FAppleARKit
 		{
 			case EAppleAnchorType::Anchor:
 			{
-				FoundGeometry->UpdateTrackedGeometry(SharedThis(this), GameThreadFrameNumber, GameThreadTimestamp, AnchorData->Transform, GetAlignmentTransform());
+				FoundGeometry->UpdateTrackedGeometry(SharedThis(this), GameThreadFrameNumber, UpdateTimestamp, AnchorData->Transform, GetAlignmentTransform());
 				for (UARPin* Pin : PinsToUpdate)
 				{
 					const FTransform Pin_LocalToTrackingTransform_PostUpdate = Pin->GetLocalToTrackingTransform_NoAlignment() * AnchorDeltaTransform;
@@ -1543,7 +1553,7 @@ void FAppleARKitSystem::SessionDidUpdateAnchors_Internal( TSharedRef<FAppleARKit
 			{
 				if (UARPlaneGeometry* PlaneGeo = Cast<UARPlaneGeometry>(FoundGeometry))
 				{
-					PlaneGeo->UpdateTrackedGeometry(SharedThis(this), GameThreadFrameNumber, GameThreadTimestamp, AnchorData->Transform, GetAlignmentTransform(), AnchorData->Center, AnchorData->Extent, AnchorData->BoundaryVerts, nullptr);
+					PlaneGeo->UpdateTrackedGeometry(SharedThis(this), GameThreadFrameNumber, UpdateTimestamp, AnchorData->Transform, GetAlignmentTransform(), AnchorData->Center, AnchorData->Extent, AnchorData->BoundaryVerts, nullptr);
 					for (UARPin* Pin : PinsToUpdate)
 					{
 						const FTransform Pin_LocalToTrackingTransform_PostUpdate = Pin->GetLocalToTrackingTransform_NoAlignment() * AnchorDeltaTransform;
@@ -1556,7 +1566,7 @@ void FAppleARKitSystem::SessionDidUpdateAnchors_Internal( TSharedRef<FAppleARKit
 			{
 				if (UARFaceGeometry* FaceGeo = Cast<UARFaceGeometry>(FoundGeometry))
 				{
-					FaceGeo->UpdateFaceGeometry(SharedThis(this), GameThreadFrameNumber, GameThreadTimestamp, AnchorData->Transform, GetAlignmentTransform(), AnchorData->BlendShapes, AnchorData->FaceVerts, AnchorData->FaceIndices, AnchorData->LeftEyeTransform, AnchorData->RightEyeTransform, AnchorData->LookAtTarget);
+					FaceGeo->UpdateFaceGeometry(SharedThis(this), GameThreadFrameNumber, UpdateTimestamp, AnchorData->Transform, GetAlignmentTransform(), AnchorData->BlendShapes, AnchorData->FaceVerts, AnchorData->FaceIndices, AnchorData->LeftEyeTransform, AnchorData->RightEyeTransform, AnchorData->LookAtTarget);
 					FaceGeo->bIsTracked = AnchorData->bIsTracked;
 					for (UARPin* Pin : PinsToUpdate)
 					{
@@ -1573,7 +1583,7 @@ void FAppleARKitSystem::SessionDidUpdateAnchors_Internal( TSharedRef<FAppleARKit
 					UARCandidateImage** CandidateImage = CandidateImages.Find(AnchorData->DetectedAnchorName);
 					ensure(CandidateImage != nullptr);
 
-					ImageAnchor->UpdateTrackedGeometry(SharedThis(this), GameThreadFrameNumber, GameThreadTimestamp, AnchorData->Transform, GetAlignmentTransform(), *CandidateImage);
+					ImageAnchor->UpdateTrackedGeometry(SharedThis(this), GameThreadFrameNumber, UpdateTimestamp, AnchorData->Transform, GetAlignmentTransform(), *CandidateImage);
 					ImageAnchor->bIsTracked = AnchorData->bIsTracked;
 					for (UARPin* Pin : PinsToUpdate)
 					{
@@ -1588,7 +1598,7 @@ void FAppleARKitSystem::SessionDidUpdateAnchors_Internal( TSharedRef<FAppleARKit
 				if (UAppleARKitEnvironmentCaptureProbe* ProbeAnchor = Cast<UAppleARKitEnvironmentCaptureProbe>(FoundGeometry))
 				{
 					// NOTE: The metal texture will be a different texture every time the cubemap is updated which requires a render resource flush
-					ProbeAnchor->UpdateEnvironmentCapture(SharedThis(this), GameThreadFrameNumber, GameThreadTimestamp, AnchorData->Transform, GetAlignmentTransform(), AnchorData->Extent, AnchorData->ProbeTexture);
+					ProbeAnchor->UpdateEnvironmentCapture(SharedThis(this), GameThreadFrameNumber, UpdateTimestamp, AnchorData->Transform, GetAlignmentTransform(), AnchorData->Extent, AnchorData->ProbeTexture);
 					for (UARPin* Pin : PinsToUpdate)
 					{
 						const FTransform Pin_LocalToTrackingTransform_PostUpdate = Pin->GetLocalToTrackingTransform_NoAlignment() * AnchorDeltaTransform;
diff --git a/Engine/Plugins/Runtime/AppleARKitFaceSupport/Source/AppleARKitFaceSupport/Private/AppleARKitLiveLinkSource.cpp b/Engine/Plugins/Runtime/AppleARKitFaceSupport/Source/AppleARKitFaceSupport/Private/AppleARKitLiveLinkSource.cpp
--- a/Engine/Plugins/Runtime/AppleARKitFaceSupport/Source/AppleARKitFaceSupport/Private/AppleARKitLiveLinkSource.cpp
+++ b/Engine/Plugins/Runtime/AppleARKitFaceSupport/Source/AppleARKitFaceSupport/Private/AppleARKitLiveLinkSource.cpp
@@ -147,9 +147,7 @@ void FAppleARKitLiveLinkSource::PublishBlendShapes(FName SubjectName, double Tim
 	{
 		static FLiveLinkFrameData LiveLinkFrame;
 
-		//LiveLinkFrame.WorldTime = Timestamp;
-		// Don't use the ARKit timestamps because there's a bug somewhere in LiveLink that doesn't interpret them properly
-		LiveLinkFrame.WorldTime = FPlatformTime::Seconds();
+		LiveLinkFrame.WorldTime = Timestamp;
 		
 		TArray<FLiveLinkCurveElement>& BlendShapes = LiveLinkFrame.CurveElements;
 
