Added compression and versioning to the AR saved world data

diff --git a/Engine/Plugins/Runtime/AppleARKit/Source/AppleARKit/Private/AppleARKitSystem.cpp b/Engine/Plugins/Runtime/AppleARKit/Source/AppleARKit/Private/AppleARKitSystem.cpp
--- a/Engine/Plugins/Runtime/AppleARKit/Source/AppleARKit/Private/AppleARKitSystem.cpp
+++ b/Engine/Plugins/Runtime/AppleARKit/Source/AppleARKit/Private/AppleARKitSystem.cpp
@@ -898,7 +898,6 @@ public:
 		}
 	}
 
-	/** Performs the call to get the world map and triggers the OnWorldMapAcquired() the completion handler */
 	void Run()
 	{
 		simd_float4x4 ARMatrix = FAppleARKitConversion::ToARKitMatrix(FTransform(Location));
@@ -950,10 +949,26 @@ public:
 		{
 			NSData* WorldNSData = [NSKeyedArchiver archivedDataWithRootObject: WorldMap];
 
-			// Copy to our TArray that will serve the data to the caller
-			uint32 SavedSize = WorldNSData.length;
-			WorldData.AddUninitialized(SavedSize);
-			FPlatformMemory::Memcpy(WorldData.GetData(), [WorldNSData bytes], SavedSize);
+			int32 UncompressedSize = WorldNSData.length;
+
+			TArray<uint8> CompressedData;
+			CompressedData.AddUninitialized(WorldNSData.length + AR_SAVE_WORLD_HEADER_SIZE);
+			uint8* Buffer = (uint8*)CompressedData.GetData();
+			// Write our magic header into our buffer
+			FARWorldSaveHeader& Header = *(FARWorldSaveHeader*)Buffer;
+			Header = FARWorldSaveHeader();
+			Header.UncompressedSize = UncompressedSize;
+			
+			// Compress the data
+			uint8* CompressInto = Buffer + AR_SAVE_WORLD_HEADER_SIZE;
+			int32 CompressedSize = UncompressedSize;
+			uint8* UncompressedData = (uint8*)[WorldNSData bytes];
+			verify(FCompression::CompressMemory((ECompressionFlags)COMPRESS_ZLIB, CompressInto, CompressedSize, UncompressedData, UncompressedSize));
+			
+			// Only copy out the amount of compressed data and the header
+			int32 CompressedSizePlusHeader = CompressedSize + AR_SAVE_WORLD_HEADER_SIZE;
+			WorldData.AddUninitialized(CompressedSizePlusHeader);
+			FPlatformMemory::Memcpy(WorldData.GetData(), CompressedData.GetData(), CompressedSizePlusHeader);
 		}
 		else
 		{
diff --git a/Engine/Plugins/Runtime/AppleARKit/Source/AppleARKit/Public/AppleARKitConversion.h b/Engine/Plugins/Runtime/AppleARKit/Source/AppleARKit/Public/AppleARKitConversion.h
--- a/Engine/Plugins/Runtime/AppleARKit/Source/AppleARKit/Public/AppleARKitConversion.h
+++ b/Engine/Plugins/Runtime/AppleARKit/Source/AppleARKit/Public/AppleARKitConversion.h
@@ -4,9 +4,38 @@
 
 #include "Math/Transform.h"
 #include "ARPin.h"
+#include "Misc/Compression.h"
 
 #include "IAppleImageUtilsPlugin.h"
 
+#define AR_SAVE_WORLD_KEY 0x505A474A
+#define AR_SAVE_WORLD_VER 1
+
+struct FARWorldSaveHeader
+{
+	uint32 Magic;
+	uint32 UncompressedSize;
+	uint8 Version;
+	
+	FARWorldSaveHeader() :
+		Magic(AR_SAVE_WORLD_KEY),
+		UncompressedSize(0),
+		Version(AR_SAVE_WORLD_VER)
+	{
+		
+	}
+
+	FARWorldSaveHeader(uint8* Header)
+	{
+		const FARWorldSaveHeader& Other = *(FARWorldSaveHeader*)Header;
+		Magic = Other.Magic;
+		UncompressedSize = Other.UncompressedSize;
+		Version = Other.Version;
+	}
+};
+
+#define AR_SAVE_WORLD_HEADER_SIZE (sizeof(FARWorldSaveHeader))
+
 struct FAppleARKitConversion
 {
 	static FORCEINLINE float ToUE4Scale()
@@ -282,7 +311,27 @@ struct FAppleARKitConversion
 
 	static ARWorldMap* ToARWorldMap(const TArray<uint8>& WorldMapData)
 	{
-		NSData* WorldNSData = [NSData dataWithBytesNoCopy: (uint8*)WorldMapData.GetData() length: WorldMapData.Num() freeWhenDone: NO];
+		uint8* Buffer = (uint8*)WorldMapData.GetData();
+		FARWorldSaveHeader InHeader(Buffer);
+		// Check for our format and reject if invalid
+		if (InHeader.Magic != AR_SAVE_WORLD_KEY || InHeader.Version != AR_SAVE_WORLD_VER)
+		{
+			return nullptr;
+		}
+
+		// Decompress the data
+		uint8* CompressedData = Buffer + AR_SAVE_WORLD_HEADER_SIZE;
+		uint32 CompressedSize = WorldMapData.Num() - AR_SAVE_WORLD_HEADER_SIZE;
+		uint32 UncompressedSize = InHeader.UncompressedSize;
+		TArray<uint8> UncompressedData;
+		UncompressedData.AddUninitialized(UncompressedSize);
+		if (!FCompression::UncompressMemory((ECompressionFlags)COMPRESS_ZLIB, UncompressedData.GetData(), UncompressedSize, CompressedData, CompressedSize))
+		{
+			return nullptr;
+		}
+		
+		// Serialize into the World map data
+		NSData* WorldNSData = [NSData dataWithBytesNoCopy: UncompressedData.GetData() length: UncompressedData.Num() freeWhenDone: NO];
 		ARWorldMap* WorldMap = [NSKeyedUnarchiver unarchiveObjectWithData: WorldNSData];
 		if (WorldMap == nullptr)
 		{
