Split some of the methods for ARKit conversion into a cpp from being all in the header
Added some logging during conversion

diff --git a/Engine/Plugins/Runtime/AppleARKit/Source/AppleARKit/Private/AppleARKitConversion.cpp b/Engine/Plugins/Runtime/AppleARKit/Source/AppleARKit/Private/AppleARKitConversion.cpp
new file mode 100644
--- /dev/null
+++ b/Engine/Plugins/Runtime/AppleARKit/Source/AppleARKit/Private/AppleARKitConversion.cpp
@@ -0,0 +1,345 @@
+// Copyright 1998-2018 Epic Games, Inc. All Rights Reserved.
+
+#include "AppleARKitConversion.h"
+#include "AppleARKitModule.h"
+
+#if SUPPORTS_ARKIT_1_0
+ARWorldAlignment FAppleARKitConversion::ToARWorldAlignment( const EARWorldAlignment& InWorldAlignment )
+{
+	switch ( InWorldAlignment )
+	{
+		case EARWorldAlignment::Gravity:
+			return ARWorldAlignmentGravity;
+
+		case EARWorldAlignment::GravityAndHeading:
+			return ARWorldAlignmentGravityAndHeading;
+
+		case EARWorldAlignment::Camera:
+			return ARWorldAlignmentCamera;
+	};
+};
+
+#pragma clang diagnostic push
+#pragma clang diagnostic ignored "-Wpartial-availability"
+
+#if SUPPORTS_ARKIT_1_5
+
+ARVideoFormat* FAppleARKitConversion::ToARVideoFormat(const FARVideoFormat& DesiredFormat, NSArray<ARVideoFormat*>* Formats)
+{
+	if (Formats != nullptr)
+	{
+		for (ARVideoFormat* Format in Formats)
+		{
+			if (Format != nullptr &&
+				DesiredFormat.FPS == Format.framesPerSecond &&
+				DesiredFormat.Width == Format.imageResolution.width &&
+				DesiredFormat.Height == Format.imageResolution.height)
+			{
+				return Format;
+			}
+		}
+	}
+	return nullptr;
+}
+
+FARVideoFormat FAppleARKitConversion::FromARVideoFormat(ARVideoFormat* Format)
+{
+	FARVideoFormat ConvertedFormat;
+	if (Format != nullptr)
+	{
+		ConvertedFormat.FPS = Format.framesPerSecond;
+		ConvertedFormat.Width = Format.imageResolution.width;
+		ConvertedFormat.Height = Format.imageResolution.height;
+	}
+	return ConvertedFormat;
+}
+
+TArray<FARVideoFormat> FAppleARKitConversion::FromARVideoFormatArray(NSArray<ARVideoFormat*>* Formats)
+{
+	TArray<FARVideoFormat> ConvertedArray;
+	if (Formats != nullptr)
+	{
+		for (ARVideoFormat* Format in Formats)
+		{
+			if (Format != nullptr)
+			{
+				ConvertedArray.Add(FromARVideoFormat(Format));
+			}
+		}
+	}
+	return ConvertedArray;
+}
+
+NSSet* FAppleARKitConversion::InitImageDetection(UARSessionConfig* SessionConfig, TMap< FString, UARCandidateImage* >& CandidateImages, TMap< FString, CGImageRef >& ConvertedCandidateImages)
+{
+	const TArray<UARCandidateImage*>& ConfigCandidateImages = SessionConfig->GetCandidateImageList();
+	if (!ConfigCandidateImages.Num())
+	{
+		return nullptr;
+	}
+
+	NSMutableSet* ConvertedImageSet = [[NSMutableSet new] autorelease];
+	for (UARCandidateImage* Candidate : ConfigCandidateImages)
+	{
+		if (Candidate != nullptr && Candidate->GetCandidateTexture() != nullptr)
+		{
+			// Store off so the session object can quickly match the anchor to our representation
+			// This stores it even if we weren't able to convert to apple's type for GC reasons
+			CandidateImages.Add(Candidate->GetFriendlyName(), Candidate);
+			// Convert our texture to an Apple compatible image type
+			CGImageRef ConvertedImage = nullptr;
+			// Avoid doing the expensive conversion work if it's in the cache already
+			CGImageRef* FoundImage = ConvertedCandidateImages.Find(Candidate->GetFriendlyName());
+			if (FoundImage != nullptr)
+			{
+				ConvertedImage = *FoundImage;
+			}
+			else
+			{
+				ConvertedImage = IAppleImageUtilsPlugin::Get().UTexture2DToCGImage(Candidate->GetCandidateTexture());
+				// If it didn't convert this time, it never will, so always store it off
+				ConvertedCandidateImages.Add(Candidate->GetFriendlyName(), ConvertedImage);
+			}
+			if (ConvertedImage != nullptr)
+			{
+				float ImageWidth = (float)Candidate->GetPhysicalWidth() / 100.f;
+				CGImagePropertyOrientation Orientation = Candidate->GetOrientation() == EARCandidateImageOrientation::Landscape ? kCGImagePropertyOrientationRight : kCGImagePropertyOrientationUp;
+
+				ARReferenceImage* ReferenceImage = [[[ARReferenceImage alloc] initWithCGImage: ConvertedImage orientation: Orientation physicalWidth: ImageWidth] autorelease];
+				ReferenceImage.name = Candidate->GetFriendlyName().GetNSString();
+				[ConvertedImageSet addObject: ReferenceImage];
+			}
+		}
+	}
+	return ConvertedImageSet;
+}
+
+void FAppleARKitConversion::InitImageDetection(UARSessionConfig* SessionConfig, ARWorldTrackingConfiguration* WorldConfig, TMap< FString, UARCandidateImage* >& CandidateImages, TMap< FString, CGImageRef >& ConvertedCandidateImages)
+{
+	if (FAppleARKitAvailability::SupportsARKit15())
+	{
+		WorldConfig.detectionImages = InitImageDetection(SessionConfig, CandidateImages, ConvertedCandidateImages);
+	}
+#if SUPPORTS_ARKIT_2_0
+	if (FAppleARKitAvailability::SupportsARKit20())
+	{
+		WorldConfig.maximumNumberOfTrackedImages = SessionConfig->GetMaxNumSimultaneousImagesTracked();
+	}
+#endif
+}
+#endif
+
+#if SUPPORTS_ARKIT_2_0
+void FAppleARKitConversion::InitImageDetection(UARSessionConfig* SessionConfig, ARImageTrackingConfiguration* ImageConfig, TMap< FString, UARCandidateImage* >& CandidateImages, TMap< FString, CGImageRef >& ConvertedCandidateImages)
+{
+	ImageConfig.trackingImages = InitImageDetection(SessionConfig, CandidateImages, ConvertedCandidateImages);
+	ImageConfig.maximumNumberOfTrackedImages = SessionConfig->GetMaxNumSimultaneousImagesTracked();
+	ImageConfig.autoFocusEnabled = SessionConfig->ShouldEnableAutoFocus();
+}
+
+AREnvironmentTexturing FAppleARKitConversion::ToAREnvironmentTexturing(EAREnvironmentCaptureProbeType CaptureType)
+{
+	switch (CaptureType)
+	{
+		case EAREnvironmentCaptureProbeType::Manual:
+		{
+			return AREnvironmentTexturingManual;
+		}
+		case EAREnvironmentCaptureProbeType::Automatic:
+		{
+			return AREnvironmentTexturingAutomatic;
+		}
+	}
+	return AREnvironmentTexturingNone;
+}
+
+ARWorldMap* FAppleARKitConversion::ToARWorldMap(const TArray<uint8>& WorldMapData)
+{
+	uint8* Buffer = (uint8*)WorldMapData.GetData();
+	FARWorldSaveHeader InHeader(Buffer);
+	// Check for our format and reject if invalid
+	if (InHeader.Magic != AR_SAVE_WORLD_KEY || InHeader.Version != AR_SAVE_WORLD_VER)
+	{
+		UE_LOG(LogAppleARKit, Log, TEXT("Failed to load the world map data from the session object due to incompatible versions: magic (0x%x), ver(%d)"), InHeader.Magic, (uint32)InHeader.Version);
+		return nullptr;
+	}
+
+	// Decompress the data
+	uint8* CompressedData = Buffer + AR_SAVE_WORLD_HEADER_SIZE;
+	uint32 CompressedSize = WorldMapData.Num() - AR_SAVE_WORLD_HEADER_SIZE;
+	uint32 UncompressedSize = InHeader.UncompressedSize;
+	TArray<uint8> UncompressedData;
+	UncompressedData.AddUninitialized(UncompressedSize);
+	if (!FCompression::UncompressMemory((ECompressionFlags)COMPRESS_ZLIB, UncompressedData.GetData(), UncompressedSize, CompressedData, CompressedSize))
+	{
+		UE_LOG(LogAppleARKit, Log, TEXT("Failed to load the world map data from the session object due to a decompression error"));
+		return nullptr;
+	}
+	
+	// Serialize into the World map data
+	NSData* WorldNSData = [NSData dataWithBytesNoCopy: UncompressedData.GetData() length: UncompressedData.Num() freeWhenDone: NO];
+	ARWorldMap* WorldMap = [NSKeyedUnarchiver unarchiveObjectWithData: WorldNSData];
+	if (WorldMap == nullptr)
+	{
+		UE_LOG(LogAppleARKit, Log, TEXT("Failed to load the world map data from the session object"));
+	}
+	return WorldMap;
+}
+
+NSSet* FAppleARKitConversion::ToARReferenceObjectSet(const TArray<UARCandidateObject*>& CandidateObjects, TMap< FString, UARCandidateObject* >& CandidateObjectMap)
+{
+	CandidateObjectMap.Empty();
+
+	if (CandidateObjects.Num() == 0)
+	{
+		return nullptr;
+	}
+
+	NSMutableSet* ConvertedObjectSet = [[NSMutableSet new] autorelease];
+	for (UARCandidateObject* Candidate : CandidateObjects)
+	{
+		if (Candidate != nullptr && Candidate->GetCandidateObjectData().Num() > 0)
+		{
+			NSData* CandidateData = [NSData dataWithBytesNoCopy: (uint8*)Candidate->GetCandidateObjectData().GetData() length: Candidate->GetCandidateObjectData().Num() freeWhenDone: NO];
+			ARReferenceObject* RefObject = [NSKeyedUnarchiver unarchiveObjectWithData: CandidateData];
+			if (RefObject != nullptr)
+			{
+				// Store off so the session object can quickly match the anchor to our representation
+				// This stores it even if we weren't able to convert to apple's type for GC reasons
+				CandidateObjectMap.Add(Candidate->GetFriendlyName(), Candidate);
+				RefObject.name = Candidate->GetFriendlyName().GetNSString();
+				[ConvertedObjectSet addObject: RefObject];
+			}
+			else
+			{
+				UE_LOG(LogAppleARKit, Log, TEXT("Failed to convert to ARReferenceObject (%s)"), *Candidate->GetFriendlyName());
+			}
+		}
+		else
+		{
+			UE_LOG(LogAppleARKit, Log, TEXT("Missing candidate object data for ARCandidateObject (%s)"), Candidate != nullptr ? *Candidate->GetFriendlyName() : TEXT("null"));
+		}
+	}
+	return ConvertedObjectSet;
+}
+#endif
+
+ARConfiguration* FAppleARKitConversion::ToARConfiguration( UARSessionConfig* SessionConfig, TMap< FString, UARCandidateImage* >& CandidateImages, TMap< FString, CGImageRef >& ConvertedCandidateImages, TMap< FString, UARCandidateObject* >& CandidateObjects )
+{
+	EARSessionType SessionType = SessionConfig->GetSessionType();
+	ARConfiguration* SessionConfiguration = nullptr;
+	switch (SessionType)
+	{
+		case EARSessionType::Orientation:
+		{
+			if (AROrientationTrackingConfiguration.isSupported == FALSE)
+			{
+				return nullptr;
+			}
+			SessionConfiguration = [AROrientationTrackingConfiguration new];
+			break;
+		}
+		case EARSessionType::World:
+		{
+			if (ARWorldTrackingConfiguration.isSupported == FALSE)
+			{
+				return nullptr;
+			}
+			ARWorldTrackingConfiguration* WorldTrackingConfiguration = [ARWorldTrackingConfiguration new];
+			WorldTrackingConfiguration.planeDetection = ARPlaneDetectionNone;
+			if ( EnumHasAnyFlags(EARPlaneDetectionMode::HorizontalPlaneDetection, SessionConfig->GetPlaneDetectionMode()))
+			{
+				WorldTrackingConfiguration.planeDetection |= ARPlaneDetectionHorizontal;
+			}
+#if SUPPORTS_ARKIT_1_5
+			if (FAppleARKitAvailability::SupportsARKit15())
+			{
+				if (EnumHasAnyFlags(EARPlaneDetectionMode::VerticalPlaneDetection, SessionConfig->GetPlaneDetectionMode()) )
+				{
+					WorldTrackingConfiguration.planeDetection |= ARPlaneDetectionVertical;
+				}
+				WorldTrackingConfiguration.autoFocusEnabled = SessionConfig->ShouldEnableAutoFocus();
+				// Add any images that wish to be detected
+				FAppleARKitConversion::InitImageDetection(SessionConfig, WorldTrackingConfiguration, CandidateImages, ConvertedCandidateImages);
+				ARVideoFormat* Format = FAppleARKitConversion::ToARVideoFormat(SessionConfig->GetDesiredVideoFormat(), ARWorldTrackingConfiguration.supportedVideoFormats);
+				if (Format != nullptr)
+				{
+					WorldTrackingConfiguration.videoFormat = Format;
+				}
+			}
+#endif
+#if SUPPORTS_ARKIT_2_0
+			if (FAppleARKitAvailability::SupportsARKit20())
+			{
+				// Check for environment capture probe types
+				WorldTrackingConfiguration.environmentTexturing = ToAREnvironmentTexturing(SessionConfig->GetEnvironmentCaptureProbeType());
+				// Load the world if requested
+				if (SessionConfig->GetWorldMapData().Num() > 0)
+				{
+					ARWorldMap* WorldMap = ToARWorldMap(SessionConfig->GetWorldMapData());
+					WorldTrackingConfiguration.initialWorldMap = WorldMap;
+					[WorldMap release];
+				}
+				// Convert any candidate objects that are to be detected
+				WorldTrackingConfiguration.detectionObjects = ToARReferenceObjectSet(SessionConfig->GetCandidateObjectList(), CandidateObjects);
+			}
+#endif
+			SessionConfiguration = WorldTrackingConfiguration;
+			break;
+		}
+		case EARSessionType::Image:
+		{
+#if SUPPORTS_ARKIT_2_0
+			if (FAppleARKitAvailability::SupportsARKit20())
+			{
+				if (ARImageTrackingConfiguration.isSupported == FALSE)
+				{
+					return nullptr;
+				}
+				ARImageTrackingConfiguration* ImageTrackingConfiguration = [ARImageTrackingConfiguration new];
+				// Add any images that wish to be detected
+				InitImageDetection(SessionConfig, ImageTrackingConfiguration, CandidateImages, ConvertedCandidateImages);
+				SessionConfiguration = ImageTrackingConfiguration;
+			}
+#endif
+			break;
+		}
+		case EARSessionType::ObjectScanning:
+		{
+#if SUPPORTS_ARKIT_2_0
+			if (FAppleARKitAvailability::SupportsARKit20())
+			{
+				if (ARObjectScanningConfiguration.isSupported == FALSE)
+				{
+					return nullptr;
+				}
+				ARObjectScanningConfiguration* ObjectScanningConfiguration = [ARObjectScanningConfiguration new];
+				if (EnumHasAnyFlags(EARPlaneDetectionMode::HorizontalPlaneDetection, SessionConfig->GetPlaneDetectionMode()))
+				{
+					ObjectScanningConfiguration.planeDetection |= ARPlaneDetectionHorizontal;
+				}
+				if (EnumHasAnyFlags(EARPlaneDetectionMode::VerticalPlaneDetection, SessionConfig->GetPlaneDetectionMode()))
+				{
+					ObjectScanningConfiguration.planeDetection |= ARPlaneDetectionVertical;
+				}
+				ObjectScanningConfiguration.autoFocusEnabled = SessionConfig->ShouldEnableAutoFocus();
+				SessionConfiguration = ObjectScanningConfiguration;
+			}
+#endif
+			break;
+		}
+		default:
+			return nullptr;
+	}
+	check(SessionConfiguration != nullptr);
+
+	// Copy / convert properties
+	SessionConfiguration.lightEstimationEnabled = SessionConfig->GetLightEstimationMode() != EARLightEstimationMode::None;
+	SessionConfiguration.providesAudioData = NO;
+	SessionConfiguration.worldAlignment = FAppleARKitConversion::ToARWorldAlignment(SessionConfig->GetWorldAlignment());
+
+	return SessionConfiguration;
+}
+
+#pragma clang diagnostic pop
+#endif
diff --git a/Engine/Plugins/Runtime/AppleARKit/Source/AppleARKit/Public/AppleARKitConversion.h b/Engine/Plugins/Runtime/AppleARKit/Source/AppleARKit/Public/AppleARKitConversion.h
--- a/Engine/Plugins/Runtime/AppleARKit/Source/AppleARKit/Public/AppleARKitConversion.h
+++ b/Engine/Plugins/Runtime/AppleARKit/Source/AppleARKit/Public/AppleARKitConversion.h
@@ -2,12 +2,21 @@
 
 #pragma once
 
+#include "AppleARKitAvailability.h"
 #include "Math/Transform.h"
 #include "ARPin.h"
 #include "Misc/Compression.h"
+#include "ARTypes.h"
+#include "ARTrackable.h"
+#include "ARSessionConfig.h"
 
 #include "IAppleImageUtilsPlugin.h"
 
+#if SUPPORTS_ARKIT_1_0
+	#import <ARKit/ARKit.h>
+#endif
+
+
 #define AR_SAVE_WORLD_KEY 0x505A474A
 #define AR_SAVE_WORLD_VER 1
 
@@ -160,340 +169,35 @@ struct FAppleARKitConversion
 		return ToFGuid( UUID );
 	}
 
-	static FORCEINLINE ARWorldAlignment ToARWorldAlignment( const EARWorldAlignment& InWorldAlignment )
-	{
-		switch ( InWorldAlignment )
-		{
-			case EARWorldAlignment::Gravity:
-				return ARWorldAlignmentGravity;
-
-			case EARWorldAlignment::GravityAndHeading:
-				return ARWorldAlignmentGravityAndHeading;
-
-			case EARWorldAlignment::Camera:
-				return ARWorldAlignmentCamera;
-		};
-	};
+	static FORCEINLINE ARWorldAlignment ToARWorldAlignment( const EARWorldAlignment& InWorldAlignment );
 
-	#pragma clang diagnostic push
-	#pragma clang diagnostic ignored "-Wpartial-availability"
+#pragma clang diagnostic push
+#pragma clang diagnostic ignored "-Wpartial-availability"
 
-	#if SUPPORTS_ARKIT_1_5
+#if SUPPORTS_ARKIT_1_5
 
-	static ARVideoFormat* ToARVideoFormat(const FARVideoFormat& DesiredFormat, NSArray<ARVideoFormat*>* Formats)
-	{
-		if (Formats != nullptr)
-		{
-			for (ARVideoFormat* Format in Formats)
-			{
-				if (Format != nullptr &&
-					DesiredFormat.FPS == Format.framesPerSecond &&
-					DesiredFormat.Width == Format.imageResolution.width &&
-					DesiredFormat.Height == Format.imageResolution.height)
-				{
-					return Format;
-				}
-			}
-		}
-		return nullptr;
-	}
+	static ARVideoFormat* ToARVideoFormat(const FARVideoFormat& DesiredFormat, NSArray<ARVideoFormat*>* Formats);
 	
-	static FARVideoFormat FromARVideoFormat(ARVideoFormat* Format)
-	{
-		FARVideoFormat ConvertedFormat;
-		if (Format != nullptr)
-		{
-			ConvertedFormat.FPS = Format.framesPerSecond;
-			ConvertedFormat.Width = Format.imageResolution.width;
-			ConvertedFormat.Height = Format.imageResolution.height;
-		}
-		return ConvertedFormat;
-	}
+	static FARVideoFormat FromARVideoFormat(ARVideoFormat* Format);
 	
-	static TArray<FARVideoFormat> FromARVideoFormatArray(NSArray<ARVideoFormat*>* Formats)
-	{
-		TArray<FARVideoFormat> ConvertedArray;
-		if (Formats != nullptr)
-		{
-			for (ARVideoFormat* Format in Formats)
-			{
-				if (Format != nullptr)
-				{
-					ConvertedArray.Add(FromARVideoFormat(Format));
-				}
-			}
-		}
-		return ConvertedArray;
-	}
-
-	static NSSet* InitImageDetection(UARSessionConfig* SessionConfig, TMap< FString, UARCandidateImage* >& CandidateImages, TMap< FString, CGImageRef >& ConvertedCandidateImages)
-	{
-		const TArray<UARCandidateImage*>& ConfigCandidateImages = SessionConfig->GetCandidateImageList();
-		if (!ConfigCandidateImages.Num())
-		{
-			return nullptr;
-		}
-
-		NSMutableSet* ConvertedImageSet = [[NSMutableSet new] autorelease];
-		for (UARCandidateImage* Candidate : ConfigCandidateImages)
-		{
-			if (Candidate != nullptr && Candidate->GetCandidateTexture() != nullptr)
-			{
-				// Store off so the session object can quickly match the anchor to our representation
-				// This stores it even if we weren't able to convert to apple's type for GC reasons
-				CandidateImages.Add(Candidate->GetFriendlyName(), Candidate);
-				// Convert our texture to an Apple compatible image type
-				CGImageRef ConvertedImage = nullptr;
-				// Avoid doing the expensive conversion work if it's in the cache already
-				CGImageRef* FoundImage = ConvertedCandidateImages.Find(Candidate->GetFriendlyName());
-				if (FoundImage != nullptr)
-				{
-					ConvertedImage = *FoundImage;
-				}
-				else
-				{
-					ConvertedImage = IAppleImageUtilsPlugin::Get().UTexture2DToCGImage(Candidate->GetCandidateTexture());
-					// If it didn't convert this time, it never will, so always store it off
-					ConvertedCandidateImages.Add(Candidate->GetFriendlyName(), ConvertedImage);
-				}
-				if (ConvertedImage != nullptr)
-				{
-					float ImageWidth = (float)Candidate->GetPhysicalWidth() / 100.f;
-					CGImagePropertyOrientation Orientation = Candidate->GetOrientation() == EARCandidateImageOrientation::Landscape ? kCGImagePropertyOrientationRight : kCGImagePropertyOrientationUp;
-
-					ARReferenceImage* ReferenceImage = [[[ARReferenceImage alloc] initWithCGImage: ConvertedImage orientation: Orientation physicalWidth: ImageWidth] autorelease];
-					ReferenceImage.name = Candidate->GetFriendlyName().GetNSString();
-					[ConvertedImageSet addObject: ReferenceImage];
-				}
-			}
-		}
-		return ConvertedImageSet;
-	}
-
-	static void InitImageDetection(UARSessionConfig* SessionConfig, ARWorldTrackingConfiguration* WorldConfig, TMap< FString, UARCandidateImage* >& CandidateImages, TMap< FString, CGImageRef >& ConvertedCandidateImages)
-	{
-		if (FAppleARKitAvailability::SupportsARKit15())
-		{
-			WorldConfig.detectionImages = InitImageDetection(SessionConfig, CandidateImages, ConvertedCandidateImages);
-		}
-	#if SUPPORTS_ARKIT_2_0
-		if (FAppleARKitAvailability::SupportsARKit20())
-		{
-			WorldConfig.maximumNumberOfTrackedImages = SessionConfig->GetMaxNumSimultaneousImagesTracked();
-		}
-	#endif
-	}
-	#endif
-
-	#if SUPPORTS_ARKIT_2_0
-	static void InitImageDetection(UARSessionConfig* SessionConfig, ARImageTrackingConfiguration* ImageConfig, TMap< FString, UARCandidateImage* >& CandidateImages, TMap< FString, CGImageRef >& ConvertedCandidateImages)
-	{
-		ImageConfig.trackingImages = InitImageDetection(SessionConfig, CandidateImages, ConvertedCandidateImages);
-		ImageConfig.maximumNumberOfTrackedImages = SessionConfig->GetMaxNumSimultaneousImagesTracked();
-		ImageConfig.autoFocusEnabled = SessionConfig->ShouldEnableAutoFocus();
-	}
-
-	static AREnvironmentTexturing ToAREnvironmentTexturing(EAREnvironmentCaptureProbeType CaptureType)
-	{
-		switch (CaptureType)
-		{
-			case EAREnvironmentCaptureProbeType::Manual:
-			{
-				return AREnvironmentTexturingManual;
-			}
-			case EAREnvironmentCaptureProbeType::Automatic:
-			{
-				return AREnvironmentTexturingAutomatic;
-			}
-		}
-		return AREnvironmentTexturingNone;
-	}
+	static TArray<FARVideoFormat> FromARVideoFormatArray(NSArray<ARVideoFormat*>* Formats);
 
-	static ARWorldMap* ToARWorldMap(const TArray<uint8>& WorldMapData)
-	{
-		uint8* Buffer = (uint8*)WorldMapData.GetData();
-		FARWorldSaveHeader InHeader(Buffer);
-		// Check for our format and reject if invalid
-		if (InHeader.Magic != AR_SAVE_WORLD_KEY || InHeader.Version != AR_SAVE_WORLD_VER)
-		{
-			return nullptr;
-		}
+	static NSSet* InitImageDetection(UARSessionConfig* SessionConfig, TMap< FString, UARCandidateImage* >& CandidateImages, TMap< FString, CGImageRef >& ConvertedCandidateImages);
 
-		// Decompress the data
-		uint8* CompressedData = Buffer + AR_SAVE_WORLD_HEADER_SIZE;
-		uint32 CompressedSize = WorldMapData.Num() - AR_SAVE_WORLD_HEADER_SIZE;
-		uint32 UncompressedSize = InHeader.UncompressedSize;
-		TArray<uint8> UncompressedData;
-		UncompressedData.AddUninitialized(UncompressedSize);
-		if (!FCompression::UncompressMemory((ECompressionFlags)COMPRESS_ZLIB, UncompressedData.GetData(), UncompressedSize, CompressedData, CompressedSize))
-		{
-			return nullptr;
-		}
-		
-		// Serialize into the World map data
-		NSData* WorldNSData = [NSData dataWithBytesNoCopy: UncompressedData.GetData() length: UncompressedData.Num() freeWhenDone: NO];
-		ARWorldMap* WorldMap = [NSKeyedUnarchiver unarchiveObjectWithData: WorldNSData];
-		if (WorldMap == nullptr)
-		{
-//			UE_LOG(LogAppleARKit, Log, TEXT("Failed to load the world map data from the session object"));
-		}
-		return WorldMap;
-	}
+	static void InitImageDetection(UARSessionConfig* SessionConfig, ARWorldTrackingConfiguration* WorldConfig, TMap< FString, UARCandidateImage* >& CandidateImages, TMap< FString, CGImageRef >& ConvertedCandidateImages);
+#endif
 
-	static NSSet* ToARReferenceObjectSet(const TArray<UARCandidateObject*>& CandidateObjects, TMap< FString, UARCandidateObject* >& CandidateObjectMap)
-	{
-		CandidateObjectMap.Empty();
+#if SUPPORTS_ARKIT_2_0
+	static void InitImageDetection(UARSessionConfig* SessionConfig, ARImageTrackingConfiguration* ImageConfig, TMap< FString, UARCandidateImage* >& CandidateImages, TMap< FString, CGImageRef >& ConvertedCandidateImages);
 
-		if (CandidateObjects.Num() == 0)
-		{
-			return nullptr;
-		}
+	static AREnvironmentTexturing ToAREnvironmentTexturing(EAREnvironmentCaptureProbeType CaptureType);
 
-		NSMutableSet* ConvertedObjectSet = [[NSMutableSet new] autorelease];
-		for (UARCandidateObject* Candidate : CandidateObjects)
-		{
-			if (Candidate != nullptr && Candidate->GetCandidateObjectData().Num() > 0)
-			{
-				NSData* CandidateData = [NSData dataWithBytesNoCopy: (uint8*)Candidate->GetCandidateObjectData().GetData() length: Candidate->GetCandidateObjectData().Num() freeWhenDone: NO];
-				ARReferenceObject* RefObject = [NSKeyedUnarchiver unarchiveObjectWithData: CandidateData];
-				if (RefObject != nullptr)
-				{
-					// Store off so the session object can quickly match the anchor to our representation
-					// This stores it even if we weren't able to convert to apple's type for GC reasons
-					CandidateObjectMap.Add(Candidate->GetFriendlyName(), Candidate);
-					RefObject.name = Candidate->GetFriendlyName().GetNSString();
-					[ConvertedObjectSet addObject: RefObject];
-				}
-				else
-				{
-//					UE_LOG(LogAppleARKit, Log, TEXT("Failed to convert to ARReferenceObject (%s)"), *Candidate->GetFriendlyName());
-				}
-			}
-			else
-			{
-//				UE_LOG(LogAppleARKit, Log, TEXT("Missing candidate object data for ARCandidateObject (%s)"), Candidate != nullptr ? *Candidate->GetFriendlyName() : TEXT("null"));
-			}
-		}
-		return ConvertedObjectSet;
-	}
-	#endif
+	static ARWorldMap* ToARWorldMap(const TArray<uint8>& WorldMapData);
 
-	static ARConfiguration* ToARConfiguration( UARSessionConfig* SessionConfig, TMap< FString, UARCandidateImage* >& CandidateImages, TMap< FString, CGImageRef >& ConvertedCandidateImages, TMap< FString, UARCandidateObject* >& CandidateObjects )
-	{
-		EARSessionType SessionType = SessionConfig->GetSessionType();
-		ARConfiguration* SessionConfiguration = nullptr;
-		switch (SessionType)
-		{
-			case EARSessionType::Orientation:
-			{
-				if (AROrientationTrackingConfiguration.isSupported == FALSE)
-				{
-					return nullptr;
-				}
-				SessionConfiguration = [AROrientationTrackingConfiguration new];
-				break;
-			}
-			case EARSessionType::World:
-			{
-				if (ARWorldTrackingConfiguration.isSupported == FALSE)
-				{
-					return nullptr;
-				}
-				ARWorldTrackingConfiguration* WorldTrackingConfiguration = [ARWorldTrackingConfiguration new];
-				WorldTrackingConfiguration.planeDetection = ARPlaneDetectionNone;
-				if ( EnumHasAnyFlags(EARPlaneDetectionMode::HorizontalPlaneDetection, SessionConfig->GetPlaneDetectionMode()))
-				{
-					WorldTrackingConfiguration.planeDetection |= ARPlaneDetectionHorizontal;
-				}
-#if SUPPORTS_ARKIT_1_5
-				if (FAppleARKitAvailability::SupportsARKit15())
-				{
-					if (EnumHasAnyFlags(EARPlaneDetectionMode::VerticalPlaneDetection, SessionConfig->GetPlaneDetectionMode()) )
-					{
-						WorldTrackingConfiguration.planeDetection |= ARPlaneDetectionVertical;
-					}
-					WorldTrackingConfiguration.autoFocusEnabled = SessionConfig->ShouldEnableAutoFocus();
-					// Add any images that wish to be detected
-					FAppleARKitConversion::InitImageDetection(SessionConfig, WorldTrackingConfiguration, CandidateImages, ConvertedCandidateImages);
-					ARVideoFormat* Format = FAppleARKitConversion::ToARVideoFormat(SessionConfig->GetDesiredVideoFormat(), ARWorldTrackingConfiguration.supportedVideoFormats);
-					if (Format != nullptr)
-					{
-						WorldTrackingConfiguration.videoFormat = Format;
-					}
-				}
-#endif
-#if SUPPORTS_ARKIT_2_0
-				if (FAppleARKitAvailability::SupportsARKit20())
-				{
-					// Check for environment capture probe types
-					WorldTrackingConfiguration.environmentTexturing = ToAREnvironmentTexturing(SessionConfig->GetEnvironmentCaptureProbeType());
-					// Load the world if requested
-					if (SessionConfig->GetWorldMapData().Num() > 0)
-					{
-						ARWorldMap* WorldMap = ToARWorldMap(SessionConfig->GetWorldMapData());
-						WorldTrackingConfiguration.initialWorldMap = WorldMap;
-						[WorldMap release];
-					}
-					// Convert any candidate objects that are to be detected
-					WorldTrackingConfiguration.detectionObjects = ToARReferenceObjectSet(SessionConfig->GetCandidateObjectList(), CandidateObjects);
-				}
+	static NSSet* ToARReferenceObjectSet(const TArray<UARCandidateObject*>& CandidateObjects, TMap< FString, UARCandidateObject* >& CandidateObjectMap);
 #endif
-				SessionConfiguration = WorldTrackingConfiguration;
-				break;
-			}
-			case EARSessionType::Image:
-			{
-#if SUPPORTS_ARKIT_2_0
-				if (FAppleARKitAvailability::SupportsARKit20())
-				{
-					if (ARImageTrackingConfiguration.isSupported == FALSE)
-					{
-						return nullptr;
-					}
-					ARImageTrackingConfiguration* ImageTrackingConfiguration = [ARImageTrackingConfiguration new];
-					// Add any images that wish to be detected
-					InitImageDetection(SessionConfig, ImageTrackingConfiguration, CandidateImages, ConvertedCandidateImages);
-					SessionConfiguration = ImageTrackingConfiguration;
-				}
-#endif
-				break;
-			}
-			case EARSessionType::ObjectScanning:
-			{
-#if SUPPORTS_ARKIT_2_0
-				if (FAppleARKitAvailability::SupportsARKit20())
-				{
-					if (ARObjectScanningConfiguration.isSupported == FALSE)
-					{
-						return nullptr;
-					}
-					ARObjectScanningConfiguration* ObjectScanningConfiguration = [ARObjectScanningConfiguration new];
-					if (EnumHasAnyFlags(EARPlaneDetectionMode::HorizontalPlaneDetection, SessionConfig->GetPlaneDetectionMode()))
-					{
-						ObjectScanningConfiguration.planeDetection |= ARPlaneDetectionHorizontal;
-					}
-					if (EnumHasAnyFlags(EARPlaneDetectionMode::VerticalPlaneDetection, SessionConfig->GetPlaneDetectionMode()))
-					{
-						ObjectScanningConfiguration.planeDetection |= ARPlaneDetectionVertical;
-					}
-					ObjectScanningConfiguration.autoFocusEnabled = SessionConfig->ShouldEnableAutoFocus();
-					SessionConfiguration = ObjectScanningConfiguration;
-				}
-#endif
-				break;
-			}
-			default:
-				return nullptr;
-		}
-		check(SessionConfiguration != nullptr);
-
-		// Copy / convert properties
-		SessionConfiguration.lightEstimationEnabled = SessionConfig->GetLightEstimationMode() != EARLightEstimationMode::None;
-		SessionConfiguration.providesAudioData = NO;
-		SessionConfiguration.worldAlignment = FAppleARKitConversion::ToARWorldAlignment(SessionConfig->GetWorldAlignment());
 
-		return SessionConfiguration;
-	}
+	static ARConfiguration* ToARConfiguration( UARSessionConfig* SessionConfig, TMap< FString, UARCandidateImage* >& CandidateImages, TMap< FString, CGImageRef >& ConvertedCandidateImages, TMap< FString, UARCandidateObject* >& CandidateObjects );
 
 	#pragma clang diagnostic pop
 #endif
