Fixed inline problem and exported FAppleARKitConversion class

diff --git a/Engine/Plugins/Runtime/AppleARKit/Source/AppleARKit/Public/AppleARKitConversion.h b/Engine/Plugins/Runtime/AppleARKit/Source/AppleARKit/Public/AppleARKitConversion.h
--- a/Engine/Plugins/Runtime/AppleARKit/Source/AppleARKit/Public/AppleARKitConversion.h
+++ b/Engine/Plugins/Runtime/AppleARKit/Source/AppleARKit/Public/AppleARKitConversion.h
@@ -45,7 +45,7 @@ struct FARWorldSaveHeader
 
 #define AR_SAVE_WORLD_HEADER_SIZE (sizeof(FARWorldSaveHeader))
 
-struct FAppleARKitConversion
+struct APPLEARKIT_API FAppleARKitConversion
 {
 	static FORCEINLINE float ToUE4Scale()
 	{
@@ -169,7 +169,7 @@ struct FAppleARKitConversion
 		return ToFGuid( UUID );
 	}
 
-	static FORCEINLINE ARWorldAlignment ToARWorldAlignment( const EARWorldAlignment& InWorldAlignment );
+	static ARWorldAlignment ToARWorldAlignment( const EARWorldAlignment& InWorldAlignment );
 
 #pragma clang diagnostic push
 #pragma clang diagnostic ignored "-Wpartial-availability"
