Adding initialization to all members in the default FAppleARKitFrame ctor

diff --git a/Engine/Plugins/Runtime/AppleARKit/Source/AppleARKit/Private/AppleARKitFrame.cpp b/Engine/Plugins/Runtime/AppleARKit/Source/AppleARKit/Private/AppleARKitFrame.cpp
--- a/Engine/Plugins/Runtime/AppleARKit/Source/AppleARKit/Private/AppleARKitFrame.cpp
+++ b/Engine/Plugins/Runtime/AppleARKit/Source/AppleARKit/Private/AppleARKitFrame.cpp
@@ -7,11 +7,18 @@
 
 // Default constructor
 FAppleARKitFrame::FAppleARKitFrame()
+	: Timestamp(0.0)
 #if SUPPORTS_ARKIT_1_0
-	: CapturedYImage(nullptr)
-	, CapturedCbCrImage( nullptr )
+	, CapturedYImage(nullptr)
+	, CapturedCbCrImage(nullptr)
+	, CameraImage(nullptr)
+	, CameraDepth(nullptr)
 	, NativeFrame(nullptr)
 #endif
+	, CapturedYImageWidth(0)
+	, CapturedYImageHeight(0)
+	, CapturedCbCrImageWidth(0)
+	, CapturedCbCrImageHeight(0)
 {
 };
 
