We can't enqueue a render command inside of the delegate handling device rotation. This ends up being unsafe due to the render thread being suspended. Instead we need to store off the new pass through image and enqueue on the next GT tick.

diff --git a/Engine/Plugins/Runtime/AR/Apple/AppleARKit/Source/AppleARKit/Private/AppleARKitSystem.cpp b/Engine/Plugins/Runtime/AR/Apple/AppleARKit/Source/AppleARKit/Private/AppleARKitSystem.cpp
--- a/Engine/Plugins/Runtime/AR/Apple/AppleARKit/Source/AppleARKit/Private/AppleARKitSystem.cpp
+++ b/Engine/Plugins/Runtime/AR/Apple/AppleARKit/Source/AppleARKit/Private/AppleARKitSystem.cpp
@@ -427,6 +427,10 @@ float FAppleARKitSystem::GetWorldToMetersScale() const
 
 void FAppleARKitSystem::OnBeginRendering_GameThread()
 {
+#if PLATFORM_MAC || PLATFORM_IOS
+    // Queue an update on the render thread
+	CameraImage->Init_RenderThread();
+#endif
 	UpdatePoses();
 }
 
@@ -1346,10 +1350,10 @@ void FAppleARKitSystem::SessionDidUpdateFrame_DelegateThread(TSharedPtr< FAppleA
 		{
 			WriteCameraImageToDisk(Frame->CameraImage);
 		}
-		// Queue an update on the render thread
+		
 		if (CameraImage != nullptr)
 		{
-			CameraImage->Init_RenderThread(Frame->CameraImage);
+			CameraImage->EnqueueNewCameraImage(Frame->CameraImage);
 		}
 #endif
 	}
diff --git a/Engine/Plugins/Runtime/AR/Apple/AppleARKit/Source/AppleARKit/Private/AppleARKitTextures.cpp b/Engine/Plugins/Runtime/AR/Apple/AppleARKit/Source/AppleARKit/Private/AppleARKitTextures.cpp
--- a/Engine/Plugins/Runtime/AR/Apple/AppleARKit/Source/AppleARKit/Private/AppleARKitTextures.cpp
+++ b/Engine/Plugins/Runtime/AR/Apple/AppleARKit/Source/AppleARKit/Private/AppleARKitTextures.cpp
@@ -189,6 +189,7 @@ UAppleARKitTextureCameraImage::UAppleARKitTextureCameraImage(const FObjectInitia
 	: Super(ObjectInitializer)
 #if PLATFORM_MAC || PLATFORM_IOS
 	, CameraImage(nullptr)
+    , NewCameraImage(nullptr)
 #endif
 {
 	ExternalTextureGuid = FGuid::NewGuid();
@@ -211,6 +212,15 @@ void UAppleARKitTextureCameraImage::BeginDestroy()
 		CFRelease(CameraImage);
 		CameraImage = nullptr;
 	}
+    
+    {
+        FScopeLock ScopeLock(&PendingImageLock);
+        if (NewCameraImage != nullptr)
+        {
+            CFRelease(NewCameraImage);
+            NewCameraImage = nullptr;
+        }
+    }
 #endif
 	Super::BeginDestroy();
 }
@@ -244,20 +254,36 @@ void UAppleARKitTextureCameraImage::Init(float InTimestamp, CVPixelBufferRef InC
 	}
 }
 
-void UAppleARKitTextureCameraImage::Init_RenderThread(CVPixelBufferRef InCameraImage)
+void UAppleARKitTextureCameraImage::Init_RenderThread()
 {
 	if (Resource != nullptr)
 	{
-		CFRetain(InCameraImage);
 		FARKitCameraImageResource* ARKitResource = static_cast<FARKitCameraImageResource*>(Resource);
 		ENQUEUE_RENDER_COMMAND(Init_RenderThread)(
-			[ARKitResource, InCameraImage](FRHICommandListImmediate&)
+			[ARKitResource, this](FRHICommandListImmediate&)
 		{
-			ARKitResource->Init_RenderThread(InCameraImage);
-			CFRelease(InCameraImage);
+			FScopeLock ScopeLock(&PendingImageLock);
+			if (NewCameraImage != nullptr)
+			{
+				ARKitResource->Init_RenderThread(NewCameraImage);
+				CFRelease(NewCameraImage);
+				NewCameraImage = nullptr;
+			}
 		});
 	}
 }
+
+void UAppleARKitTextureCameraImage::EnqueueNewCameraImage(CVPixelBufferRef InCameraImage)
+{
+	FScopeLock ScopeLock(&PendingImageLock);
+	if (NewCameraImage != nullptr)
+	{
+		CFRelease(NewCameraImage);
+	}
+
+	NewCameraImage = InCameraImage;
+	CFRetain(NewCameraImage);
+}
 #endif
 
 UAppleARKitTextureCameraDepth::UAppleARKitTextureCameraDepth(const FObjectInitializer& ObjectInitializer)
diff --git a/Engine/Plugins/Runtime/AR/Apple/AppleARKit/Source/AppleARKit/Public/AppleARKitTextures.h b/Engine/Plugins/Runtime/AR/Apple/AppleARKit/Source/AppleARKit/Public/AppleARKitTextures.h
--- a/Engine/Plugins/Runtime/AR/Apple/AppleARKit/Source/AppleARKit/Public/AppleARKitTextures.h
+++ b/Engine/Plugins/Runtime/AR/Apple/AppleARKit/Source/AppleARKit/Public/AppleARKitTextures.h
@@ -32,7 +32,10 @@ public:
 	virtual void Init(float InTimestamp, CVPixelBufferRef InCameraImage);
 
 	/** Queues the conversion on the render thread */
-	virtual void Init_RenderThread(CVPixelBufferRef InCameraImage);
+	virtual void Init_RenderThread();
+
+    /** Store off the camera image for a later resource update */
+	void EnqueueNewCameraImage(CVPixelBufferRef InCameraImage);
 #endif
 
 	virtual EAppleTextureType GetTextureType() const override { return EAppleTextureType::PixelBuffer; }
@@ -50,6 +53,8 @@ private:
 #if PLATFORM_MAC || PLATFORM_IOS
 	/** The Apple specific representation of the ar camera image */
 	CVPixelBufferRef CameraImage;
+	CVPixelBufferRef NewCameraImage;
+	FCriticalSection PendingImageLock;
 #endif
 };
 
