We need to adjust the pass through camera uvs to account for the difference between the screen/render target aspect ratio and the camera texture's aspect ratio.

diff --git a/Engine/Plugins/Runtime/AR/Apple/AppleARKit/Source/AppleARKit/Private/AppleARKitSystem.cpp b/Engine/Plugins/Runtime/AR/Apple/AppleARKit/Source/AppleARKit/Private/AppleARKitSystem.cpp
--- a/Engine/Plugins/Runtime/AR/Apple/AppleARKit/Source/AppleARKit/Private/AppleARKitSystem.cpp
+++ b/Engine/Plugins/Runtime/AR/Apple/AppleARKit/Source/AppleARKit/Private/AppleARKitSystem.cpp
@@ -143,12 +143,21 @@ private:
 	
 	virtual void PreRenderViewFamily_RenderThread(FRHICommandListImmediate& RHICmdList, FSceneViewFamily& InViewFamily) override
 	{
+		// Grab the latest frame from ARKit
+		{
+			FScopeLock ScopeLock(&ARKitSystem.FrameLock);
+			ARKitSystem.RenderThreadFrame = ARKitSystem.LastReceivedFrame;
+		}
+
 		FDefaultXRCamera::PreRenderViewFamily_RenderThread(RHICmdList, InViewFamily);
 	}
 	
 	virtual void PostRenderBasePass_RenderThread(FRHICommandListImmediate& RHICmdList, FSceneView& InView) override
 	{
-		VideoOverlay.RenderVideoOverlay_RenderThread(RHICmdList, InView);
+		if (ARKitSystem.RenderThreadFrame.IsValid())
+		{
+			VideoOverlay.RenderVideoOverlay_RenderThread(RHICmdList, InView, *ARKitSystem.RenderThreadFrame, ARKitSystem.DeviceOrientation);
+		}
 	}
 	
 	virtual bool GetPassthroughCameraUVs_RenderThread(TArray<FVector2D>& OutUVs) override
diff --git a/Engine/Plugins/Runtime/AR/Apple/AppleARKit/Source/AppleARKit/Private/AppleARKitVideoOverlay.cpp b/Engine/Plugins/Runtime/AR/Apple/AppleARKit/Source/AppleARKit/Private/AppleARKitVideoOverlay.cpp
--- a/Engine/Plugins/Runtime/AR/Apple/AppleARKit/Source/AppleARKit/Private/AppleARKitVideoOverlay.cpp
+++ b/Engine/Plugins/Runtime/AR/Apple/AppleARKit/Source/AppleARKit/Private/AppleARKitVideoOverlay.cpp
@@ -158,39 +158,96 @@ private:
 IMPLEMENT_MATERIAL_SHADER_TYPE(template<>, FARKitCameraOverlayPS<true>, TEXT("/Engine/Private/PostProcessMaterialShaders.usf"), TEXT("MainPS_ES2"), SF_Pixel);
 IMPLEMENT_MATERIAL_SHADER_TYPE(template<>, FARKitCameraOverlayPS<false>, TEXT("/Engine/Private/PostProcessMaterialShaders.usf"), TEXT("MainPS_VideoOverlay"), SF_Pixel);
 
-void FAppleARKitVideoOverlay::RenderVideoOverlay_RenderThread(FRHICommandListImmediate& RHICmdList, const FSceneView& InView)
+void FAppleARKitVideoOverlay::RenderVideoOverlay_RenderThread(FRHICommandListImmediate& RHICmdList, const FSceneView& InView, FAppleARKitFrame& Frame, const EDeviceScreenOrientation DeviceOrientation)
 {
 	if (RenderingOverlayMaterial == nullptr || !RenderingOverlayMaterial->IsValidLowLevel())
 	{
 		return;
 	}
 
-	if (VertexBufferRHI == nullptr || !VertexBufferRHI.IsValid())
+	if (OverlayVertexBufferRHI[0] == nullptr || !OverlayVertexBufferRHI[0].IsValid())
 	{
+		const FVector2D ViewSize(InView.UnconstrainedViewRect.Max.X, InView.UnconstrainedViewRect.Max.Y);
+
+		FVector2D CameraSize = Frame.Camera.ImageResolution;
+		if ((ViewSize.X > ViewSize.Y) != (CameraSize.X > CameraSize.Y))
+		{
+			CameraSize = FVector2D(CameraSize.Y, CameraSize.X);
+		}
+
+		const float CameraAspectRatio = CameraSize.X / CameraSize.Y;
+		const float ViewAspectRatio = ViewSize.X / ViewSize.Y;
+		const float ViewAspectRatioLandscape = (ViewSize.X > ViewSize.Y) ? ViewAspectRatio : ViewSize.Y / ViewSize.X;
+
+		float UVOffsetAmount = 0.0f;
+		if (!FMath::IsNearlyEqual(ViewAspectRatio, CameraAspectRatio))
+		{
+			if (ViewAspectRatio > CameraAspectRatio)
+			{
+				UVOffsetAmount = 0.5f * (1.0f - (CameraAspectRatio / ViewAspectRatio));
+			}
+			else
+			{
+				UVOffsetAmount = 0.5f * (1.0f - (ViewAspectRatio / CameraAspectRatio));
+			}
+		}
+
+		UVOffset = (ViewAspectRatioLandscape <= Frame.Camera.GetAspectRatio()) ? FVector2D(UVOffsetAmount, 0.0f) : FVector2D(0.0f, UVOffsetAmount);
+
 		// Setup vertex buffer
-		TResourceArray<FFilterVertex, VERTEXBUFFER_ALIGNMENT> Vertices;
-		Vertices.SetNumUninitialized(4);
+		const FVector4 Positions[] =
+		{
+			FVector4(0.0f, 1.0f, 0.0f, 1.0f),
+			FVector4(0.0f, 0.0f, 0.0f, 1.0f),
+			FVector4(1.0f, 1.0f, 0.0f, 1.0f),
+			FVector4(1.0f, 0.0f, 0.0f, 1.0f)
+		};
+
+		const FVector2D UVs[] =
+		{
+			// Landscape
+			FVector2D(UVOffset.X, 1.0f - UVOffset.Y),
+			FVector2D(UVOffset.X, UVOffset.Y),
+			FVector2D(1.0f - UVOffset.X, 1.0f - UVOffset.Y),
+			FVector2D(1.0f - UVOffset.X, UVOffset.Y),
+            
+			// Portrait
+			FVector2D(UVOffset.Y, 1.0f - UVOffset.X),
+            FVector2D(UVOffset.Y, UVOffset.X),
+			FVector2D(1.0f - UVOffset.Y, 1.0f - UVOffset.X),
+			FVector2D(1.0f - UVOffset.Y, UVOffset.X),
+		};
+
+		uint32 UVIndex = 0;
+		for (uint32 OrientationIter = 0; OrientationIter < 2; ++OrientationIter)
+		{
+			TResourceArray<FFilterVertex, VERTEXBUFFER_ALIGNMENT> Vertices;
+			Vertices.SetNumUninitialized(4);
+
+			Vertices[0].Position = Positions[0];
+			Vertices[0].UV = UVs[UVIndex];
 
-		Vertices[0].Position = FVector4(0.f, 0.f, 0.f, 1.f);
-		Vertices[0].UV = FVector2D(0.f, 0.f);
+			Vertices[1].Position = Positions[1];
+			Vertices[1].UV = UVs[UVIndex + 1];
 
-		Vertices[1].Position = FVector4(1.f, 0.f, 0.f, 1.f);
-		Vertices[1].UV = FVector2D(1.f, 0.f);
+			Vertices[2].Position = Positions[2];
+			Vertices[2].UV = UVs[UVIndex + 2];
 
-		Vertices[2].Position = FVector4(1.f, 1.f, 0.f, 1.f);
-		Vertices[2].UV = FVector2D(1.f, 1.f);
+			Vertices[3].Position = Positions[3];
+			Vertices[3].UV = UVs[UVIndex + 3];
 
-		Vertices[3].Position = FVector4(0.f, 1.f, 0.f, 1.f);
-		Vertices[3].UV = FVector2D(0.f, 1.f);
+			UVIndex += 4;
 
-		FRHIResourceCreateInfo CreateInfoVB(&Vertices);
-		VertexBufferRHI = RHICreateVertexBuffer(Vertices.GetResourceDataSize(), BUF_Static, CreateInfoVB);
+			FRHIResourceCreateInfo CreateInfoVB(&Vertices);
+			OverlayVertexBufferRHI[OrientationIter] = RHICreateVertexBuffer(Vertices.GetResourceDataSize(), BUF_Static, CreateInfoVB);
+		}
 	}
 
+
 	if (IndexBufferRHI == nullptr || !IndexBufferRHI.IsValid())
 	{
 		// Setup index buffer
-		const uint16 Indices[] = { 0, 1, 2, 2, 3, 0 };
+		const uint16 Indices[] = { 0, 1, 2, 2, 1, 3 };
 
 		TResourceArray<uint16, INDEXBUFFER_ALIGNMENT> IndexBuffer;
 		const uint32 NumIndices = ARRAY_COUNT(Indices);
@@ -269,29 +326,87 @@ void FAppleARKitVideoOverlay::RenderVideoOverlay_RenderThread(FRHICommandListImm
 		PixelShaderPtr->SetParameters(RHICmdList, InView, RenderingOverlayMaterial->GetRenderProxy());
 	}
 
-	RHICmdList.SetStreamSource(0, VertexBufferRHI, 0);
-	RHICmdList.DrawIndexedPrimitive(
-		IndexBufferRHI,
-		/*BaseVertexIndex=*/ 0,
-		/*MinIndex=*/ 0,
-		/*NumVertices=*/ 4,
-		/*StartIndex=*/ 0,
-		/*NumPrimitives=*/ 2,
-		/*NumInstances=*/ 1
-	);
+	FRHIVertexBuffer* VertexBufferRHI = nullptr;
+	switch (DeviceOrientation)
+	{
+		case EDeviceScreenOrientation::LandscapeRight:
+		case EDeviceScreenOrientation::LandscapeLeft:
+			VertexBufferRHI = OverlayVertexBufferRHI[0];
+			break;
+
+		case EDeviceScreenOrientation::Portrait:
+		case EDeviceScreenOrientation::PortraitUpsideDown:
+			VertexBufferRHI = OverlayVertexBufferRHI[1];
+			break;
+
+		default:
+			VertexBufferRHI = OverlayVertexBufferRHI[0];
+			break;
+	}
+
+
+	if (VertexBufferRHI && IndexBufferRHI.IsValid())
+	{
+		RHICmdList.SetStreamSource(0, VertexBufferRHI, 0);
+		RHICmdList.DrawIndexedPrimitive(
+			IndexBufferRHI,
+			/*BaseVertexIndex=*/ 0,
+			/*MinIndex=*/ 0,
+			/*NumVertices=*/ 4,
+			/*StartIndex=*/ 0,
+			/*NumPrimitives=*/ 2,
+			/*NumInstances=*/ 1
+		);
+	}
 }
 
 
-bool FAppleARKitVideoOverlay::GetPassthroughCameraUVs_RenderThread(TArray<FVector2D>& OutUVs, const EDeviceScreenOrientation)
+bool FAppleARKitVideoOverlay::GetPassthroughCameraUVs_RenderThread(TArray<FVector2D>& OutUVs, const EDeviceScreenOrientation DeviceOrientation)
 {
-	OutUVs.Reset(4);
-
-	new(OutUVs) FVector2D(0.f, 0.f);
-	new(OutUVs) FVector2D(1.f, 0.f);
-	new(OutUVs) FVector2D(1.0f, 1.0f);
-	new(OutUVs) FVector2D(0.f, 1.0f);
+#if SUPPORTS_ARKIT_1_0
+	if (OverlayVertexBufferRHI[0] != nullptr && OverlayVertexBufferRHI[0].IsValid())
+	{
+		OutUVs.SetNumUninitialized(4);
 
-	return true;
+		switch (DeviceOrientation)
+		{
+		case EDeviceScreenOrientation::LandscapeRight:
+			OutUVs[1] = FVector2D(UVOffset.X, 1.0f - UVOffset.Y);
+			OutUVs[0] = FVector2D(UVOffset.X, UVOffset.Y);
+			OutUVs[3] = FVector2D(1.0f - UVOffset.X, 1.0f - UVOffset.Y);
+			OutUVs[2] = FVector2D(1.0f - UVOffset.X, UVOffset.Y);
+			return true;
+
+		case EDeviceScreenOrientation::LandscapeLeft:
+            OutUVs[1] = FVector2D(UVOffset.X, 1.0f - UVOffset.Y);
+            OutUVs[0] = FVector2D(UVOffset.X, UVOffset.Y);
+            OutUVs[3] = FVector2D(1.0f - UVOffset.X, 1.0f - UVOffset.Y);
+            OutUVs[2] = FVector2D(1.0f - UVOffset.X, UVOffset.Y);
+			return true;
+
+		case EDeviceScreenOrientation::Portrait:
+			OutUVs[1] = FVector2D(UVOffset.Y, 1.0f - UVOffset.X);
+			OutUVs[0] = FVector2D(UVOffset.Y, UVOffset.X);
+			OutUVs[3] = FVector2D(1.0f - UVOffset.Y, 1.0f - UVOffset.X);
+			OutUVs[2] = FVector2D(1.0f - UVOffset.Y, UVOffset.X);
+			return true;
+
+		case EDeviceScreenOrientation::PortraitUpsideDown:
+            OutUVs[1] = FVector2D(UVOffset.Y, 1.0f - UVOffset.X);
+            OutUVs[0] = FVector2D(UVOffset.Y, UVOffset.X);
+            OutUVs[3] = FVector2D(1.0f - UVOffset.Y, 1.0f - UVOffset.X);
+            OutUVs[2] = FVector2D(1.0f - UVOffset.Y, UVOffset.X);
+			return true;
+
+		default:
+			return false;
+		}
+	}
+	else
+#endif
+	{
+		return false;
+	}
 }
 
 void FAppleARKitVideoOverlay::AddReferencedObjects(FReferenceCollector& Collector)
diff --git a/Engine/Plugins/Runtime/AR/Apple/AppleARKit/Source/AppleARKit/Private/AppleARKitVideoOverlay.h b/Engine/Plugins/Runtime/AR/Apple/AppleARKit/Source/AppleARKit/Private/AppleARKitVideoOverlay.h
--- a/Engine/Plugins/Runtime/AR/Apple/AppleARKit/Source/AppleARKit/Private/AppleARKitVideoOverlay.h
+++ b/Engine/Plugins/Runtime/AR/Apple/AppleARKit/Source/AppleARKit/Private/AppleARKitVideoOverlay.h
@@ -35,7 +35,7 @@ public:
 
 	void SetCameraTexture(UARTextureCameraImage* InCameraImage);
 
-	void RenderVideoOverlay_RenderThread(FRHICommandListImmediate& RHICmdList, const FSceneView& InView);
+	void RenderVideoOverlay_RenderThread(FRHICommandListImmediate& RHICmdList, const FSceneView& InView, struct FAppleARKitFrame& Frame, const EDeviceScreenOrientation DeviceOrientation);
 	bool GetPassthroughCameraUVs_RenderThread(TArray<FVector2D>& OutUVs, const EDeviceScreenOrientation DeviceOrientation);
 
 	void SetOverlayTexture(UARTextureCameraImage* InCameraImage);
@@ -47,6 +47,10 @@ private:
 
 	UMaterialInstanceDynamic* MID_CameraOverlay;
 	UMaterialInterface* RenderingOverlayMaterial;
+
+	// Cache UV offset to be used by GetPassthroughCameraUVs_RenderThread
+	FVector2D UVOffset;
+
+	FVertexBufferRHIRef OverlayVertexBufferRHI[2];
 	FIndexBufferRHIRef IndexBufferRHI;
-	FVertexBufferRHIRef VertexBufferRHI;
 };
