Adding beta ARKit v3 support from task stream

diff --git a/Engine/Plugins/Runtime/AR/Apple/AppleARKit/AppleARKit.uplugin b/Engine/Plugins/Runtime/AR/Apple/AppleARKit/AppleARKit.uplugin
--- a/Engine/Plugins/Runtime/AR/Apple/AppleARKit/AppleARKit.uplugin
+++ b/Engine/Plugins/Runtime/AR/Apple/AppleARKit/AppleARKit.uplugin
@@ -22,12 +22,24 @@
 			"Type": "Runtime",
 			"LoadingPhase": "PostConfigInit",
 			"WhitelistPlatforms" : [ "Mac", "iOS", "Win64" ]
+		},
+		{
+			"Name": "AppleARKitPoseTrackingLiveLink",
+			"Type": "Runtime",
+			"LoadingPhase": "PostConfigInit",
+			"WhitelistPlatforms" : [ "Mac", "iOS", "Win64" ]
+
 		}
 	],
     "Plugins": [
         {
             "Name": "AppleImageUtils",
+            "Enabled": true
+		},
+		{
+            "Name": "LiveLink",
             "Enabled": true
         }
+
     ]
 }
diff --git a/Engine/Plugins/Runtime/AR/Apple/AppleARKit/Source/AppleARKit/Private/AppleARKitConversion.cpp b/Engine/Plugins/Runtime/AR/Apple/AppleARKit/Source/AppleARKit/Private/AppleARKitConversion.cpp
--- a/Engine/Plugins/Runtime/AR/Apple/AppleARKit/Source/AppleARKit/Private/AppleARKitConversion.cpp
+++ b/Engine/Plugins/Runtime/AR/Apple/AppleARKit/Source/AppleARKit/Private/AppleARKitConversion.cpp
@@ -2,6 +2,18 @@
 
 #include "AppleARKitConversion.h"
 #include "AppleARKitModule.h"
+#include "HAL/PlatformMisc.h"
+
+template<typename TEnum>
+static FString GetEnumValueAsString(const FString& Name, TEnum Value)
+{
+	if (const UEnum* EnumClass = FindObject<UEnum>(ANY_PACKAGE, *Name, true))
+	{
+		return EnumClass->GetNameByValue((int64)Value).ToString();
+	}
+
+	return FString("Invalid");
+}
 
 #if SUPPORTS_ARKIT_1_0
 ARWorldAlignment FAppleARKitConversion::ToARWorldAlignment( const EARWorldAlignment& InWorldAlignment )
@@ -340,22 +352,280 @@ ARConfiguration* FAppleARKitConversion::ToARConfiguration( UARSessionConfig* Ses
 				ObjectScanningConfiguration.autoFocusEnabled = SessionConfig->ShouldEnableAutoFocus();
 				SessionConfiguration = ObjectScanningConfiguration;
 			}
+#endif
+			break;
+		}
+		case EARSessionType::PoseTracking:
+		{
+#if SUPPORTS_ARKIT_3_0
+			if (FAppleARKitAvailability::SupportsARKit30())
+			{
+				if (ARBodyTrackingConfiguration.isSupported == FALSE)
+				{
+					return nullptr;
+				}
+				
+				ARBodyTrackingConfiguration* BodyTrackingConfiguration = [ARBodyTrackingConfiguration new];
+				BodyTrackingConfiguration.planeDetection = ARPlaneDetectionNone;
+				if (EnumHasAnyFlags(EARPlaneDetectionMode::HorizontalPlaneDetection, SessionConfig->GetPlaneDetectionMode()))
+				{
+					BodyTrackingConfiguration.planeDetection |= ARPlaneDetectionHorizontal;
+				}
+				
+				if (EnumHasAnyFlags(EARPlaneDetectionMode::VerticalPlaneDetection, SessionConfig->GetPlaneDetectionMode()) )
+				{
+					BodyTrackingConfiguration.planeDetection |= ARPlaneDetectionVertical;
+				}
+				BodyTrackingConfiguration.autoFocusEnabled = SessionConfig->ShouldEnableAutoFocus();
+				
+				// Add any images that wish to be detected
+				FAppleARKitConversion::InitImageDetection(SessionConfig, BodyTrackingConfiguration, CandidateImages, ConvertedCandidateImages);
+				
+				ARVideoFormat* Format = FAppleARKitConversion::ToARVideoFormat(SessionConfig->GetDesiredVideoFormat(), ARBodyTrackingConfiguration.supportedVideoFormats);
+				if (Format != nullptr)
+				{
+					BodyTrackingConfiguration.videoFormat = Format;
+				}
+				
+				// Check for environment capture probe types
+				BodyTrackingConfiguration.environmentTexturing = ToAREnvironmentTexturing(SessionConfig->GetEnvironmentCaptureProbeType());
+				// Load the world if requested
+				if (SessionConfig->GetWorldMapData().Num() > 0)
+				{
+					ARWorldMap* WorldMap = ToARWorldMap(SessionConfig->GetWorldMapData());
+					BodyTrackingConfiguration.initialWorldMap = WorldMap;
+					[WorldMap release];
+				}
+				
+				SessionConfiguration = BodyTrackingConfiguration;
+			}
 #endif
 			break;
 		}
 		default:
 			return nullptr;
 	}
-    if (SessionConfiguration != nullptr)
+	if (SessionConfiguration != nullptr)
+	{
+		// Copy / convert properties
+		SessionConfiguration.lightEstimationEnabled = SessionConfig->GetLightEstimationMode() != EARLightEstimationMode::None;
+		SessionConfiguration.providesAudioData = NO;
+		SessionConfiguration.worldAlignment = FAppleARKitConversion::ToARWorldAlignment(SessionConfig->GetWorldAlignment());
+		
+#if SUPPORTS_ARKIT_3_0
+		// Enable additional frame semantics for ARKit 3.0
+		if (FAppleARKitAvailability::SupportsARKit30())
+		{
+			const EARSessionTrackingFeature SessionTrackingFeature = SessionConfig->GetEnabledSessionTrackingFeature();
+			if (SessionTrackingFeature != EARSessionTrackingFeature::None)
+			{
+				if (IsSessionTrackingFeatureSupported(SessionType, SessionTrackingFeature))
+				{
+					SessionConfiguration.frameSemantics = ToARFrameSemantics(SessionTrackingFeature);
+				}
+				else
+				{
+					UE_LOG(LogAppleARKit, Error, TEXT("Session type [%s] doesn't support the required session feature: [%s]!"),
+						   *GetEnumValueAsString<>(TEXT("EARSessionType"), SessionType),
+						   *GetEnumValueAsString<>(TEXT("EARSessionTrackingFeature"), SessionTrackingFeature));
+				}
+			}
+		}
+#endif
+	}
+	
+	return SessionConfiguration;
+}
+
+bool FAppleARKitConversion::IsSessionTrackingFeatureSupported(EARSessionType SessionType, EARSessionTrackingFeature SessionTrackingFeature)
+{
+#if SUPPORTS_ARKIT_3_0
+	if (FAppleARKitAvailability::SupportsARKit30())
+	{
+		const ARFrameSemantics Semantics = ToARFrameSemantics(SessionTrackingFeature);
+		if (Semantics != ARFrameSemanticNone)
+		{
+			switch (SessionType)
+			{
+				case EARSessionType::Orientation:
+					return [AROrientationTrackingConfiguration supportsFrameSemantics: Semantics];
+				
+				case EARSessionType::World:
+					return [ARWorldTrackingConfiguration supportsFrameSemantics: Semantics];
+				
+				case EARSessionType::Face:
+					return [ARFaceTrackingConfiguration supportsFrameSemantics: Semantics];
+				
+				case EARSessionType::Image:
+					return [ARImageTrackingConfiguration supportsFrameSemantics: Semantics];
+				
+				case EARSessionType::ObjectScanning:
+					return [ARObjectScanningConfiguration supportsFrameSemantics: Semantics];
+					
+				case EARSessionType::PoseTracking:
+					return [ARBodyTrackingConfiguration supportsFrameSemantics: Semantics];
+			}
+		}
+	}
+#endif
+	
+	return false;
+}
+
+#if SUPPORTS_ARKIT_3_0
+void FAppleARKitConversion::InitImageDetection(UARSessionConfig* SessionConfig, ARBodyTrackingConfiguration* BodyTrackingConfig, TMap< FString, UARCandidateImage* >& CandidateImages, TMap< FString, CGImageRef >& ConvertedCandidateImages)
+{
+	if (FAppleARKitAvailability::SupportsARKit15())
+	{
+		BodyTrackingConfig.detectionImages = InitImageDetection(SessionConfig, CandidateImages, ConvertedCandidateImages);
+	}
+	
+	if (FAppleARKitAvailability::SupportsARKit20())
+	{
+		BodyTrackingConfig.maximumNumberOfTrackedImages = SessionConfig->GetMaxNumSimultaneousImagesTracked();
+	}
+}
+
+ARFrameSemantics FAppleARKitConversion::ToARFrameSemantics(EARSessionTrackingFeature SessionTrackingFeature)
+{
+	static const TMap<EARSessionTrackingFeature, ARFrameSemantics> SessionTrackingFeatureToFrameSemantics =
+	{
+		{ EARSessionTrackingFeature::None, ARFrameSemanticNone },
+		{ EARSessionTrackingFeature::PoseDetection2D, ARFrameSemanticBodyDetection },
+		{ EARSessionTrackingFeature::PersonSegmentation, ARFrameSemanticPersonSegmentation },
+		{ EARSessionTrackingFeature::PersonSegmentationWithDepth, ARFrameSemanticPersonSegmentationWithDepth },
+	};
+	
+	if (SessionTrackingFeatureToFrameSemantics.Contains(SessionTrackingFeature))
+	{
+		return SessionTrackingFeatureToFrameSemantics[SessionTrackingFeature];
+	}
+	
+	return ARFrameSemanticNone;
+}
+
+void FAppleARKitConversion::ToSkeletonDefinition(const ARSkeletonDefinition* InARSkeleton, FARSkeletonDefinition& OutSkeletonDefinition)
+{
+	check(InARSkeleton);
+	
+	// TODO: these values should not change over time so they can be cached somewhere
+	
+	const int NumJoints = InARSkeleton.jointCount;
+	
+	OutSkeletonDefinition.NumJoints = NumJoints;
+	OutSkeletonDefinition.JointNames.AddUninitialized(NumJoints);
+	OutSkeletonDefinition.ParentIndices.AddUninitialized(NumJoints);
+	
+	for (int Index = 0; Index < NumJoints; ++Index)
+	{
+		OutSkeletonDefinition.JointNames[Index] = *FString(InARSkeleton.jointNames[Index]);
+		OutSkeletonDefinition.ParentIndices[Index] = InARSkeleton.parentIndices[Index].intValue;
+	}
+}
+
+FARPose2D FAppleARKitConversion::ToARPose2D(const ARBody2D* InARPose2D)
+{
+	check(InARPose2D);
+	
+    const EDeviceScreenOrientation ScreenOrientation = FPlatformMisc::GetDeviceOrientation();
+	
+	FARPose2D Pose2D;
+	if (FAppleARKitAvailability::SupportsARKit30())
+	{
+		const ARSkeleton2D* Skeleton2D = InARPose2D.skeleton;
+		ToSkeletonDefinition(Skeleton2D.definition, Pose2D.SkeletonDefinition);
+		const int NumJoints = Pose2D.SkeletonDefinition.NumJoints;
+		
+		Pose2D.JointLocations.AddUninitialized(NumJoints);
+		Pose2D.IsJointTracked.AddUninitialized(NumJoints);
+		
+		for (int Index = 0; Index < NumJoints; ++Index)
+		{
+			const bool bIsTracked = [Skeleton2D isJointTracked: Index];
+			Pose2D.IsJointTracked[Index] = bIsTracked;
+			
+            if (bIsTracked)
+            {
+                FVector2D OriginalLandmark(Skeleton2D.jointLandmarks[Index][0], Skeleton2D.jointLandmarks[Index][1]);
+				
+                switch (ScreenOrientation)
+                {
+                    case EDeviceScreenOrientation::Portrait:
+                        Pose2D.JointLocations[Index] = FVector2D(1.f - OriginalLandmark.Y, OriginalLandmark.X);
+                        break;
+                        
+                    case EDeviceScreenOrientation::PortraitUpsideDown:
+                        Pose2D.JointLocations[Index] = FVector2D(OriginalLandmark.Y, OriginalLandmark.X);
+                        break;
+                        
+                    case EDeviceScreenOrientation::LandscapeLeft:
+                        Pose2D.JointLocations[Index] = FVector2D(1.f - OriginalLandmark.X, 1.f - OriginalLandmark.Y);
+                        break;
+                        
+                    case EDeviceScreenOrientation::LandscapeRight:
+                        Pose2D.JointLocations[Index] = OriginalLandmark;
+                        break;
+                        
+                    default:
+                        Pose2D.JointLocations[Index] = OriginalLandmark;
+                }
+            }
+            else
+            {
+                Pose2D.JointLocations[Index] = FVector2D::ZeroVector;
+            }
+		}
+	}
+	
+	return Pose2D;
+}
+
+FARPose3D FAppleARKitConversion::ToARPose3D(const ARSkeleton3D* Skeleton3D, bool bIdentityForUntracked)
+{
+	check(Skeleton3D);
+	
+	FARPose3D Pose3D;
+	if (FAppleARKitAvailability::SupportsARKit30())
+	{
+		ToSkeletonDefinition(Skeleton3D.definition, Pose3D.SkeletonDefinition);
+		const int NumJoints = Pose3D.SkeletonDefinition.NumJoints;
+		
+		Pose3D.JointTransforms.AddUninitialized(NumJoints);
+		Pose3D.IsJointTracked.AddUninitialized(NumJoints);
+		Pose3D.JointTransformSpace = EARJointTransformSpace::Model;
+		
+		for (int Index = 0; Index < NumJoints; ++Index)
+		{
+			const bool bIsTracked = [Skeleton3D isJointTracked: Index];
+			Pose3D.IsJointTracked[Index] = bIsTracked;
+			if (bIsTracked || !bIdentityForUntracked)
+			{
+				Pose3D.JointTransforms[Index] = ToFTransform(Skeleton3D.jointModelTransforms[Index]);
+			}
+			else
+			{
+				Pose3D.JointTransforms[Index] = FTransform::Identity;
+			}
+		}
+	}
+	
+	return Pose3D;
+}
+
+FARPose3D FAppleARKitConversion::ToARPose3D(const ARBodyAnchor* InARBodyAnchor)
+{
+	check(InARBodyAnchor);
+	
+	FARPose3D Pose3D;
+	if (FAppleARKitAvailability::SupportsARKit30())
     {
-        // Copy / convert properties
-        SessionConfiguration.lightEstimationEnabled = SessionConfig->GetLightEstimationMode() != EARLightEstimationMode::None;
-        SessionConfiguration.providesAudioData = NO;
-        SessionConfiguration.worldAlignment = FAppleARKitConversion::ToARWorldAlignment(SessionConfig->GetWorldAlignment());
-    }
-    
-    return SessionConfiguration;
+		Pose3D = ToARPose3D(InARBodyAnchor.skeleton, true);
+	}
+	
+	return Pose3D;
 }
 
+#endif
+
 #pragma clang diagnostic pop
 #endif
diff --git a/Engine/Plugins/Runtime/AR/Apple/AppleARKit/Source/AppleARKit/Private/AppleARKitFrame.cpp b/Engine/Plugins/Runtime/AR/Apple/AppleARKit/Source/AppleARKit/Private/AppleARKitFrame.cpp
--- a/Engine/Plugins/Runtime/AR/Apple/AppleARKit/Source/AppleARKit/Private/AppleARKitFrame.cpp
+++ b/Engine/Plugins/Runtime/AR/Apple/AppleARKit/Source/AppleARKit/Private/AppleARKitFrame.cpp
@@ -36,7 +36,7 @@ EARWorldMappingState ToEARWorldMappingState(ARWorldMappingStatus MapStatus)
 
 #if SUPPORTS_ARKIT_1_0
 
-FAppleARKitFrame::FAppleARKitFrame( ARFrame* InARFrame )
+FAppleARKitFrame::FAppleARKitFrame( ARFrame* InARFrame, const FVector2D MinCameraUV, const FVector2D MaxCameraUV )
 	: Camera( InARFrame.camera )
 	, LightEstimate( InARFrame.lightEstimate )
 	, WorldMappingState(EARWorldMappingState::NotAvailable)
@@ -70,6 +70,39 @@ FAppleARKitFrame::FAppleARKitFrame( ARFrame* InARFrame )
 		WorldMappingState = ToEARWorldMappingState(InARFrame.worldMappingStatus);
 	}
 #endif
+
+#if SUPPORTS_ARKIT_3_0
+	if (FAppleARKitAvailability::SupportsARKit30())
+	{
+		if (InARFrame.detectedBody)
+		{
+			Tracked2DPose = FAppleARKitConversion::ToARPose2D(InARFrame.detectedBody);
+			
+			// Convert the joint location from the normalized arkit camera space to UE4's normalized screen space
+			const FVector2D UVSize = MaxCameraUV - MinCameraUV;
+			for (int Index = 0; Index < Tracked2DPose.JointLocations.Num(); ++Index)
+			{
+				if (Tracked2DPose.IsJointTracked[Index])
+				{
+					FVector2D& JointLocation = Tracked2DPose.JointLocations[Index];
+					JointLocation = (JointLocation - MinCameraUV) / UVSize;
+				}
+			}
+		}
+		
+		if (InARFrame.segmentationBuffer)
+		{
+			SegmentationBuffer = InARFrame.segmentationBuffer;
+			CFRetain(SegmentationBuffer);
+		}
+		
+		if (InARFrame.estimatedDepthData)
+		{
+			EstimatedDepthData = InARFrame.estimatedDepthData;
+			CFRetain(EstimatedDepthData);
+		}
+	}
+#endif
 }
 
 FAppleARKitFrame::FAppleARKitFrame( const FAppleARKitFrame& Other )
@@ -79,6 +112,7 @@ FAppleARKitFrame::FAppleARKitFrame( const FAppleARKitFrame& Other )
 	, Camera( Other.Camera )
 	, LightEstimate( Other.LightEstimate )
 	, WorldMappingState(Other.WorldMappingState)
+	, Tracked2DPose(Other.Tracked2DPose)
 {
 	if(Other.NativeFrame != nullptr)
 	{
@@ -101,6 +135,17 @@ FAppleARKitFrame::~FAppleARKitFrame()
 	{
 		CFRelease((CFTypeRef)NativeFrame);
 	}
+#if SUPPORTS_ARKIT_3_0
+	if (SegmentationBuffer)
+	{
+		CFRelease(SegmentationBuffer);
+	}
+	
+	if (EstimatedDepthData)
+	{
+		CFRelease(EstimatedDepthData);
+	}
+#endif
 }
 
 FAppleARKitFrame& FAppleARKitFrame::operator=( const FAppleARKitFrame& Other )
@@ -132,6 +177,20 @@ FAppleARKitFrame& FAppleARKitFrame::operator=( const FAppleARKitFrame& Other )
 		NativeFrame = (void*)CFRetain((CFTypeRef)Other.NativeFrame);
 	}
 	
+#if SUPPORTS_ARKIT_3_0
+	if (SegmentationBuffer)
+	{
+		CFRelease(SegmentationBuffer);
+		SegmentationBuffer = nullptr;
+	}
+	
+	if (EstimatedDepthData)
+	{
+		CFRelease(EstimatedDepthData);
+		EstimatedDepthData = nullptr;
+	}
+#endif
+	
 	// Member-wise copy
 	Timestamp = Other.Timestamp;
 	Camera = Other.Camera;
@@ -139,6 +198,8 @@ FAppleARKitFrame& FAppleARKitFrame::operator=( const FAppleARKitFrame& Other )
 	WorldMappingState = Other.WorldMappingState;
 
 	NativeFrame = Other.NativeFrame;
+	
+	Tracked2DPose = Other.Tracked2DPose;
 
 	return *this;
 }
diff --git a/Engine/Plugins/Runtime/AR/Apple/AppleARKit/Source/AppleARKit/Private/AppleARKitSessionDelegate.cpp b/Engine/Plugins/Runtime/AR/Apple/AppleARKit/Source/AppleARKit/Private/AppleARKitSessionDelegate.cpp
--- a/Engine/Plugins/Runtime/AR/Apple/AppleARKit/Source/AppleARKit/Private/AppleARKitSessionDelegate.cpp
+++ b/Engine/Plugins/Runtime/AR/Apple/AppleARKit/Source/AppleARKit/Private/AppleARKitSessionDelegate.cpp
@@ -11,6 +11,8 @@
 @implementation FAppleARKitSessionDelegate
 {
 	FAppleARKitSystem* _AppleARKitSystem;
+	TArray<FVector2D> PassthroughCameraUVs;
+	FCriticalSection CameraUVsLock;
 }
 
 
@@ -30,8 +32,30 @@
 
 - (void)session:(ARSession *)session didUpdateFrame:(ARFrame *)frame 
 {
+	// Update the camera UVs
+	TSharedPtr<IXRCamera, ESPMode::ThreadSafe> Camera = _AppleARKitSystem->GetXRCamera(0);
+	if (Camera)
+	{
+		ENQUEUE_RENDER_COMMAND(UpdateCameraUVsCommand)([self, Camera](FRHICommandListImmediate& RHICmdList)
+		{
+			FScopeLock ScopeLock(&CameraUVsLock);
+			Camera->GetPassthroughCameraUVs_RenderThread(PassthroughCameraUVs);
+		});
+	}
+	
+	FVector2D MinCameraUV(0.f, 0.f);
+	FVector2D MaxCameraUV(1.f, 1.f);
+	{
+		FScopeLock ScopeLock(&CameraUVsLock);
+		if (PassthroughCameraUVs.Num() == 4)
+		{
+			MinCameraUV = PassthroughCameraUVs[0];
+			MaxCameraUV = PassthroughCameraUVs[3];
+		}
+	}
+	
 	// Bundle results into FAppleARKitFrame
-	TSharedPtr< FAppleARKitFrame, ESPMode::ThreadSafe > AppleARKitFrame( new FAppleARKitFrame( frame ) );
+	TSharedPtr< FAppleARKitFrame, ESPMode::ThreadSafe > AppleARKitFrame( new FAppleARKitFrame( frame, MinCameraUV, MaxCameraUV ) );
 	
 	// Pass result to session
 	_AppleARKitSystem->SessionDidUpdateFrame_DelegateThread( AppleARKitFrame );
diff --git a/Engine/Plugins/Runtime/AR/Apple/AppleARKit/Source/AppleARKit/Private/AppleARKitSystem.cpp b/Engine/Plugins/Runtime/AR/Apple/AppleARKit/Source/AppleARKit/Private/AppleARKitSystem.cpp
--- a/Engine/Plugins/Runtime/AR/Apple/AppleARKit/Source/AppleARKit/Private/AppleARKitSystem.cpp
+++ b/Engine/Plugins/Runtime/AR/Apple/AppleARKit/Source/AppleARKit/Private/AppleARKitSystem.cpp
@@ -26,6 +26,7 @@
 
 // To separate out the face ar library linkage from standard ar apps
 #include "AppleARKitFaceSupport.h"
+#include "AppleARKitPoseTrackingLiveLink.h"
 
 // For orientation changed
 #include "Misc/CoreDelegates.h"
@@ -93,6 +94,11 @@ public:
 		VideoOverlay.SetOverlayTexture(InCameraImage);
 	}
 	
+	void SetEnablePersonOcclusion(bool bEnable)
+    {
+        VideoOverlay.SetEnablePersonOcclusion(bEnable);
+    }
+	
 private:
 	//~ FDefaultXRCamera
 	void OverrideFOV(float& InOutFOV)
@@ -156,7 +162,7 @@ private:
 	{
 		if (ARKitSystem.RenderThreadFrame.IsValid())
 		{
-			VideoOverlay.RenderVideoOverlay_RenderThread(RHICmdList, InView, *ARKitSystem.RenderThreadFrame, ARKitSystem.DeviceOrientation);
+			VideoOverlay.RenderVideoOverlay_RenderThread(RHICmdList, InView, *ARKitSystem.RenderThreadFrame, ARKitSystem.DeviceOrientation, ARKitSystem.GetWorldToMetersScale());
 		}
 	}
 	
@@ -209,6 +215,7 @@ FAppleARKitSystem::FAppleARKitSystem()
 , CameraDepth(nullptr)
 , LastTrackedGeometry_DebugId(0)
 , FaceARSupport(nullptr)
+, PoseTrackingARLiveLink(nullptr)
 , TimecodeProvider(nullptr)
 {
 	// See Initialize(), as we need access to SharedThis()
@@ -228,6 +235,7 @@ void FAppleARKitSystem::Shutdown()
 	if (Session != nullptr)
 	{
 		FaceARSupport = nullptr;
+		PoseTrackingARLiveLink = nullptr;
 		[Session pause];
 		Session.delegate = nullptr;
 		[Session release];
@@ -236,6 +244,9 @@ void FAppleARKitSystem::Shutdown()
 #endif
 	CameraDepth = nullptr;
 	CameraImage = nullptr;
+	
+	PersonSegmentationImage = nullptr;
+	PersonSegmentationDepthImage = nullptr;
 }
 
 void FAppleARKitSystem::CheckForFaceARSupport(UARSessionConfig* InSessionConfig)
@@ -256,6 +267,26 @@ void FAppleARKitSystem::CheckForFaceARSupport(UARSessionConfig* InSessionConfig)
 	}
 }
 
+void FAppleARKitSystem::CheckForPoseTrackingARLiveLink(UARSessionConfig* InSessionConfig)
+{
+#if SUPPORTS_ARKIT_3_0	
+	if (InSessionConfig->GetSessionType() != EARSessionType::PoseTracking)
+	{
+		// Clear the face ar support so we don't forward to it
+		PoseTrackingARLiveLink = nullptr;
+		return;
+	}
+
+	// We need to get the face support from the factory method, which is a modular feature to avoid dependencies
+	TArray<IAppleARKitPoseTrackingLiveLink*> Impls = IModularFeatures::Get().GetModularFeatureImplementations<IAppleARKitPoseTrackingLiveLink>("AppleARKitPoseTrackingLiveLink");
+	if (ensureAlwaysMsgf(Impls.Num() > 0, TEXT("Body Tracking AR session has been requested but the body tracking ar plugin is not enabled")))
+	{
+		PoseTrackingARLiveLink = Impls[0];
+		ensureAlwaysMsgf(PoseTrackingARLiveLink != nullptr, TEXT("Body Tracking AR session has been requested but the body tracking ar plugin is not enabled"));
+	}
+#endif
+}
+
 FName FAppleARKitSystem::GetSystemName() const
 {
 	static const FName AppleARKitSystemName(TEXT("AppleARKit"));
@@ -376,7 +407,32 @@ void FAppleARKitSystem::UpdateFrame()
 			if (GameThreadFrame->CameraDepth != nullptr)
 			{
 				// Reuse the UObjects because otherwise the time between GCs causes ARKit to be starved of resources
-                CameraDepth->Init(FPlatformTime::Seconds(), GameThreadFrame->CameraDepth);
+				CameraDepth->Init(FPlatformTime::Seconds(), GameThreadFrame->CameraDepth);
+			}
+#endif
+			
+#if SUPPORTS_ARKIT_3_0
+			FAppleARKitXRCamera* Camera = GetARKitXRCamera();
+			check(Camera);
+			
+			if (GameThreadFrame->SegmentationBuffer)
+			{
+				if (!PersonSegmentationImage)
+				{
+					PersonSegmentationImage = NewObject<UAppleARKitTextureCameraImage>();
+				}
+				PersonSegmentationImage->Init(FPlatformTime::Seconds(), GameThreadFrame->SegmentationBuffer);
+				PersonSegmentationImage->EnqueueNewCameraImage(GameThreadFrame->SegmentationBuffer);
+			}
+			
+			if (GameThreadFrame->EstimatedDepthData)
+			{
+				if (!PersonSegmentationDepthImage)
+				{
+					PersonSegmentationDepthImage = NewObject<UAppleARKitTextureCameraImage>();
+				}
+				PersonSegmentationDepthImage->Init(FPlatformTime::Seconds(), GameThreadFrame->EstimatedDepthData);
+				PersonSegmentationDepthImage->EnqueueNewCameraImage(GameThreadFrame->EstimatedDepthData);
 			}
 #endif
 		}
@@ -428,6 +484,11 @@ TSharedPtr<class IXRCamera, ESPMode::ThreadSafe> FAppleARKitSystem::GetXRCamera(
 	return XRCamera;
 }
 
+FAppleARKitXRCamera* FAppleARKitSystem::GetARKitXRCamera()
+{
+	return (FAppleARKitXRCamera*)GetXRCamera(0).Get();
+}
+
 float FAppleARKitSystem::GetWorldToMetersScale() const
 {
 	// @todo arkit FAppleARKitSystem::GetWorldToMetersScale needs a real scale somehow
@@ -439,6 +500,16 @@ void FAppleARKitSystem::OnBeginRendering_GameThread()
 #if PLATFORM_MAC || PLATFORM_IOS
     // Queue an update on the render thread
 	CameraImage->Init_RenderThread();
+	
+	if (PersonSegmentationImage)
+	{
+		PersonSegmentationImage->Init_RenderThread();
+	}
+	
+	if (PersonSegmentationDepthImage)
+	{
+		PersonSegmentationDepthImage->Init_RenderThread();
+	}
 #endif
 	UpdatePoses();
 }
@@ -807,6 +878,12 @@ bool FAppleARKitSystem::OnIsTrackingTypeSupported(EARSessionType SessionType) co
 		{
 			return ARObjectScanningConfiguration.isSupported == TRUE;
 		}
+#endif
+#if SUPPORTS_ARKIT_3_0
+		case EARSessionType::PoseTracking:
+		{
+			return ARBodyTrackingConfiguration.isSupported == TRUE;
+		}
 #endif
 	}
 #endif
@@ -1126,6 +1203,8 @@ void FAppleARKitSystem::AddReferencedObjects( FReferenceCollector& Collector )
 	Collector.AddReferencedObjects( CandidateImages );
 	Collector.AddReferencedObjects( CandidateObjects );
 	Collector.AddReferencedObject( TimecodeProvider );
+	Collector.AddReferencedObject( PersonSegmentationImage );
+	Collector.AddReferencedObject( PersonSegmentationDepthImage );
 
 	if(LightEstimate)
 	{
@@ -1174,7 +1253,8 @@ void FAppleARKitSystem::SetupCameraTextures()
 	{
 		CameraImage = NewObject<UAppleARKitTextureCameraImage>();
 		CameraImage->Init(FPlatformTime::Seconds(), nullptr);
-		FAppleARKitXRCamera* Camera = (FAppleARKitXRCamera*)XRCamera.Get();
+		FAppleARKitXRCamera* Camera = GetARKitXRCamera();
+		check(Camera);
 		Camera->SetOverlayTexture(CameraImage);
 	}
 	if (CameraDepth == nullptr)
@@ -1190,6 +1270,11 @@ bool FAppleARKitSystem::Run(UARSessionConfig* SessionConfig)
 	TimecodeProvider = UAppleARKitSettings::GetTimecodeProvider();
 
 	SetupCameraTextures();
+	
+	if (FAppleARKitXRCamera* Camera = GetARKitXRCamera())
+    {
+        Camera->SetEnablePersonOcclusion(SessionConfig->bUsePersonSegmentationForOcclusion);
+    }
 
 	{
 		// Clear out any existing frames since they aren't valid anymore
@@ -1216,6 +1301,7 @@ bool FAppleARKitSystem::Run(UARSessionConfig* SessionConfig)
 
 		ARConfiguration* Configuration = nullptr;
 		CheckForFaceARSupport(SessionConfig);
+		CheckForPoseTrackingARLiveLink(SessionConfig);
 		if (FaceARSupport == nullptr)
 		{
 			Configuration = FAppleARKitConversion::ToARConfiguration(SessionConfig, CandidateImages, ConvertedCandidateImages, CandidateObjects);
@@ -1228,7 +1314,7 @@ bool FAppleARKitSystem::Run(UARSessionConfig* SessionConfig)
 		// Not all session types are supported by all devices
 		if (Configuration == nullptr)
 		{
-			UE_LOG(LogAppleARKit, Log, TEXT("The requested session type is not supported by this device"));
+			UE_LOG(LogAppleARKit, Error, TEXT("The requested session type is not supported by this device"));
 			return false;
 		}
 
@@ -1377,6 +1463,7 @@ void FAppleARKitSystem::SessionDidFailWithError_DelegateThread(const FString& Er
 
 TArray<int32> FAppleARKitAnchorData::FaceIndices;
 bool FAppleARKitAnchorData::bGenerateGeometry = false;
+TSharedPtr<FARPose3D> FAppleARKitAnchorData::BodyRefPose = TSharedPtr<FARPose3D>();
 
 static TSharedPtr<FAppleARKitAnchorData> MakeAnchorData( ARAnchor* Anchor, double Timestamp, uint32 FrameNumber )
 {
@@ -1492,6 +1579,23 @@ static TSharedPtr<FAppleARKitAnchorData> MakeAnchorData( ARAnchor* Anchor, doubl
 			  FString(ObjectAnchor.referenceObject.name)
 		  );
 	}
+#endif
+#if SUPPORTS_ARKIT_3_0
+	else if (FAppleARKitAvailability::SupportsARKit30() && [Anchor isKindOfClass:[ARBodyAnchor class]])
+	{
+		ARBodyAnchor* BodyAnchor = (ARBodyAnchor*)Anchor;
+
+		if (!FAppleARKitAnchorData::BodyRefPose)
+		{
+			FAppleARKitAnchorData::BodyRefPose = MakeShared<FARPose3D>(FAppleARKitConversion::ToARPose3D(BodyAnchor.skeleton.definition.neutralBodySkeleton3D, false));
+		}
+
+		NewAnchor = MakeShared<FAppleARKitAnchorData>(
+			  FAppleARKitConversion::ToFGuid(BodyAnchor.identifier),
+			  FAppleARKitConversion::ToFTransform(BodyAnchor.transform),
+			  FAppleARKitConversion::ToARPose3D(BodyAnchor)
+		  );
+	}
 #endif
 	else
 	{
@@ -1603,6 +1707,11 @@ void FAppleARKitSystem::SessionDidAddAnchors_Internal( TSharedRef<FAppleARKitAnc
 		FaceARSupport->PublishLiveLinkData(AnchorData);
 	}
 
+	if (PoseTrackingARLiveLink != nullptr && AnchorData->AnchorType == EAppleAnchorType::PoseAnchor)
+	{
+		PoseTrackingARLiveLink->PublishLiveLinkData(AnchorData);
+	}
+
 	FString NewAnchorDebugName;
 	UARTrackedGeometry* NewGeometry = nullptr;
 	switch (AnchorData->AnchorType)
@@ -1709,6 +1818,14 @@ void FAppleARKitSystem::SessionDidAddAnchors_Internal( TSharedRef<FAppleARKitAnc
 			NewGeometry = NewTrackedObject;
 			break;
 		}
+		case EAppleAnchorType::PoseAnchor:
+		{
+			NewAnchorDebugName = FString::Printf(TEXT("POSE-%02d"), LastTrackedGeometry_DebugId++);
+			UARTrackedPose* NewTrackedPose = NewObject<UARTrackedPose>();
+			NewTrackedPose->UpdateTrackedPose(ARComponent.ToSharedRef(), AnchorData->FrameNumber, AnchorData->Timestamp, AnchorData->Transform, GetARCompositionComponent()->GetAlignmentTransform(), AnchorData->TrackedPose);
+			NewGeometry = NewTrackedPose;
+			break;
+		}
 	}
 	check(NewGeometry != nullptr);
 
@@ -1739,6 +1856,11 @@ void FAppleARKitSystem::SessionDidUpdateAnchors_Internal( TSharedRef<FAppleARKit
 		FaceARSupport->PublishLiveLinkData(AnchorData);
 	}
 
+	if (PoseTrackingARLiveLink != nullptr && AnchorData->AnchorType == EAppleAnchorType::PoseAnchor)
+	{
+		PoseTrackingARLiveLink->PublishLiveLinkData(AnchorData);
+	}
+
 	UARTrackedGeometry** GeometrySearchResult = TrackedGeometries.Find(AnchorData->AnchorGUID);
 	if (GeometrySearchResult != nullptr)
 	{
@@ -1843,6 +1965,21 @@ void FAppleARKitSystem::SessionDidUpdateAnchors_Internal( TSharedRef<FAppleARKit
 				}
 				break;
 			}
+			case EAppleAnchorType::PoseAnchor:
+			{
+				if (UARTrackedPose* TrackedPose = Cast<UARTrackedPose>(FoundGeometry))
+				{
+					TrackedPose->UpdateTrackedPose(ARComponent.ToSharedRef(), AnchorData->FrameNumber, AnchorData->Timestamp, AnchorData->Transform, GetARCompositionComponent()->GetAlignmentTransform(), AnchorData->TrackedPose);
+					
+					// TODO: why is this duplicated for every anchor type??
+					for (UARPin* Pin : PinsToUpdate)
+					{
+						const FTransform Pin_LocalToTrackingTransform_PostUpdate = Pin->GetLocalToTrackingTransform_NoAlignment() * AnchorDeltaTransform;
+						Pin->OnTransformUpdated(Pin_LocalToTrackingTransform_PostUpdate);
+					}
+				}
+				break;
+			}
 		}
 		// Trigger the delegate so anyone listening can take action
 		TriggerOnTrackableUpdatedDelegates(FoundGeometry);
@@ -2018,6 +2155,30 @@ void FAppleARKitSystem::WriteCameraImageToDisk(CVPixelBufferRef PixelBuffer)
 }
 #endif
 
+bool FAppleARKitSystem::OnIsSessionTrackingFeatureSupported(EARSessionType SessionType, EARSessionTrackingFeature SessionTrackingFeature) const
+{
+	return FAppleARKitConversion::IsSessionTrackingFeatureSupported(SessionType, SessionTrackingFeature);
+}
+
+TArray<FARPose2D> FAppleARKitSystem::OnGetTracked2DPose() const
+{
+	if (GameThreadFrame && GameThreadFrame->Tracked2DPose.SkeletonDefinition.NumJoints > 0)
+	{
+		return { GameThreadFrame->Tracked2DPose };
+	}
+	
+	return {};
+}
+
+UARTextureCameraImage* FAppleARKitSystem::OnGetPersonSegmentationImage() const
+{
+	return PersonSegmentationImage;
+}
+
+UARTextureCameraImage* FAppleARKitSystem::OnGetPersonSegmentationDepthImage() const
+{
+	return PersonSegmentationDepthImage;
+}
 
 namespace AppleARKitSupport
 {
@@ -2087,7 +2248,13 @@ FString UAppleARKitSettings::GetFaceTrackingLogDir()
 bool UAppleARKitSettings::IsLiveLinkEnabledForFaceTracking()
 {
 	FScopeLock ScopeLock(&CriticalSection);
-	return bEnableLiveLinkForFaceTracking;
+	return LivelinkTrackingType == ELivelinkTrackingType::FaceTracking;
+}
+
+bool UAppleARKitSettings::IsLiveLinkEnabledForPoseTracking()
+{
+	FScopeLock ScopeLock(&CriticalSection);
+	return LivelinkTrackingType == ELivelinkTrackingType::PoseTracking;
 }
 
 bool UAppleARKitSettings::IsFaceTrackingLoggingEnabled()
diff --git a/Engine/Plugins/Runtime/AR/Apple/AppleARKit/Source/AppleARKit/Private/AppleARKitSystem.h b/Engine/Plugins/Runtime/AR/Apple/AppleARKit/Source/AppleARKit/Private/AppleARKitSystem.h
--- a/Engine/Plugins/Runtime/AR/Apple/AppleARKit/Source/AppleARKit/Private/AppleARKitSystem.h
+++ b/Engine/Plugins/Runtime/AR/Apple/AppleARKit/Source/AppleARKit/Private/AppleARKitSystem.h
@@ -8,6 +8,7 @@
 #include "AppleARKitTextures.h"
 #include "Kismet/BlueprintPlatformLibrary.h"
 #include "AppleARKitFaceSupport.h"
+#include "AppleARKitPoseTrackingLiveLink.h"
 
 // ARKit
 #if SUPPORTS_ARKIT_1_0
@@ -86,7 +87,11 @@ protected:
 	virtual TArray<FARVideoFormat> OnGetSupportedVideoFormats(EARSessionType SessionType) const override;
 	virtual TArray<FVector> OnGetPointCloud() const override;
 	virtual bool OnAddRuntimeCandidateImage(UARSessionConfig* SessionConfig, UTexture2D* CandidateTexture, FString FriendlyName, float PhysicalWidth) override { return true; };
-
+	
+	virtual bool OnIsSessionTrackingFeatureSupported(EARSessionType SessionType, EARSessionTrackingFeature SessionTrackingFeature) const override;
+	virtual TArray<FARPose2D> OnGetTracked2DPose() const override;
+	virtual UARTextureCameraImage* OnGetPersonSegmentationImage() const override;
+	virtual UARTextureCameraImage* OnGetPersonSegmentationDepthImage() const override;
 	//~IARSystemSupport
 
 private:
@@ -101,7 +106,8 @@ private:
 	/** Asynchronously writes a JPEG to disk */
 	void WriteCameraImageToDisk(CVPixelBufferRef PixelBuffer);
 #endif
-
+	class FAppleARKitXRCamera* GetARKitXRCamera();
+	
 public:
 	// Session delegate callbacks
 	void SessionDidUpdateFrame_DelegateThread( TSharedPtr< FAppleARKitFrame, ESPMode::ThreadSafe > Frame );
@@ -143,6 +149,9 @@ private:
 
 	/** Creates or clears the face ar support object if face ar has been requested */
 	void CheckForFaceARSupport(UARSessionConfig* InSessionConfig);
+
+	/** Creates or clears the pose tracking ar support object if face ar has been requested */
+	void CheckForPoseTrackingARLiveLink(UARSessionConfig* InSessionConfig);
 	
 	/** Updates the ARKit perf counters */
 	void UpdateARKitPerfStats();
@@ -179,6 +188,8 @@ private:
 	UAppleARKitTextureCameraDepth* CameraDepth;
 	TMap< FString, UARCandidateImage* > CandidateImages;
 	TMap< FString, UARCandidateObject* > CandidateObjects;
+	UAppleARKitTextureCameraImage* PersonSegmentationImage = nullptr;
+	UAppleARKitTextureCameraImage* PersonSegmentationDepthImage = nullptr;
 	// ...
 	// PROPERTIES REPORTED TO FGCObject
 	//
@@ -198,6 +209,9 @@ private:
 	// The object that is handling face support if present
 	IAppleARKitFaceSupport* FaceARSupport;
 
+	// The object that is handling pose tracking livelink if present
+	IAppleARKitPoseTrackingLiveLink* PoseTrackingARLiveLink;
+
 	/** The time code provider to use when tagging time stamps */
 	UTimecodeProvider* TimecodeProvider;
 };
diff --git a/Engine/Plugins/Runtime/AR/Apple/AppleARKit/Source/AppleARKit/Private/AppleARKitTextures.cpp b/Engine/Plugins/Runtime/AR/Apple/AppleARKit/Source/AppleARKit/Private/AppleARKitTextures.cpp
--- a/Engine/Plugins/Runtime/AR/Apple/AppleARKit/Source/AppleARKit/Private/AppleARKitTextures.cpp
+++ b/Engine/Plugins/Runtime/AR/Apple/AppleARKit/Source/AppleARKit/Private/AppleARKitTextures.cpp
@@ -31,6 +31,13 @@ public:
 
 	virtual ~FARKitCameraImageResource()
 	{
+#if PLATFORM_MAC || PLATFORM_IOS
+		if (ImageContext)
+		{
+			CFRelease(ImageContext);
+			ImageContext = nullptr;
+		}
+#endif
 	}
 
 	virtual void InitRHI() override
@@ -70,7 +77,12 @@ public:
 			id<MTLTexture> UnderlyingMetalTexture = (id<MTLTexture>)DecodedTextureRef->GetNativeResource();
 
 			// Do the conversion on the GPU
-			[[CIContext context] render: RotatedImage toMTLTexture: UnderlyingMetalTexture commandBuffer: nil bounds: ImageExtent colorSpace: ColorSpaceRef];
+			if (!ImageContext)
+			{
+				ImageContext = [CIContext context];
+				CFRetain(ImageContext);
+			}
+			[ImageContext render: RotatedImage toMTLTexture: UnderlyingMetalTexture commandBuffer: nil bounds: ImageExtent colorSpace: ColorSpaceRef];
 
 			// Now that the conversion is done, we can get rid of our refs
 			[Image release];
@@ -176,6 +188,9 @@ private:
 #if PLATFORM_MAC || PLATFORM_IOS
 	/** The raw camera image from ARKit which is to be converted into a RGBA image */
 	CVPixelBufferRef CameraImage;
+	
+	/** The cached image context that's reused between frames */
+	CIContext* ImageContext = nullptr;
 #endif
 	/** The texture that we actually render with which is populated via the GPU conversion process */
 	FTexture2DRHIRef DecodedTextureRef;
@@ -442,6 +457,11 @@ public:
 	
 	virtual ~FARMetalResource()
 	{
+		if (ImageContext)
+		{
+			CFRelease(ImageContext);
+			ImageContext = nullptr;
+		}
 	}
 	
 	/**
@@ -513,9 +533,12 @@ public:
 		id<MTLTexture> UnderlyingMetalTexture = (id<MTLTexture>)Cubemap->GetNativeResource();
 		id<MTLTexture> OurCubeFaceMetalTexture = [UnderlyingMetalTexture newTextureViewWithPixelFormat: MTLPixelFormatBGRA8Unorm textureType: MTLTextureType2D levels: NSMakeRange(0, 1) slices: NSMakeRange(OurCubeIndex, 1)];
 
-		CIContext* Context = [CIContext context];
-
-		[Context render: RotatedCubefaceImage toMTLTexture: OurCubeFaceMetalTexture commandBuffer: nil bounds: CubefaceImage.extent colorSpace: CubefaceImage.colorSpace];
+		if (!ImageContext)
+		{
+			ImageContext = [CIContext context];
+			CFRetain(ImageContext);
+		}
+		[ImageContext render: RotatedCubefaceImage toMTLTexture: OurCubeFaceMetalTexture commandBuffer: nil bounds: CubefaceImage.extent colorSpace: CubefaceImage.colorSpace];
 
 		[CubefaceImage release];
 		[CubeFaceMetalTexture release];
@@ -548,6 +571,8 @@ private:
 	FTextureCubeRHIRef EnvCubemapTextureRHIRef;
 	
 	const UAppleARKitEnvironmentCaptureProbeTexture* Owner;
+	
+	CIContext* ImageContext = nullptr;
 };
 
 #endif
@@ -571,3 +596,251 @@ void UAppleARKitEnvironmentCaptureProbeTexture::BeginDestroy()
 #endif
 	Super::BeginDestroy();
 }
+
+UAppleARKitOcclusionTexture::UAppleARKitOcclusionTexture(const FObjectInitializer& ObjectInitializer)
+	: Super(ObjectInitializer)
+{
+	SRGB = false;
+}
+
+void UAppleARKitOcclusionTexture::BeginDestroy()
+{
+#if PLATFORM_MAC || PLATFORM_IOS
+	if (MetalTexture)
+	{
+		CFRelease(MetalTexture);
+		MetalTexture = nullptr;
+	}
+#endif
+	
+	Super::BeginDestroy();
+}
+
+#if PLATFORM_MAC || PLATFORM_IOS
+class FOcclusionTextureResource : public FTextureResource
+{
+public:
+	FOcclusionTextureResource(UAppleARKitOcclusionTexture* InOwner)
+		: Owner(InOwner)
+	{
+		bGreyScaleFormat = false;
+		bSRGB = InOwner->SRGB;
+	}
+	
+	virtual ~FOcclusionTextureResource()
+	{
+		if (ImageContext)
+		{
+			CFRelease(ImageContext);
+			ImageContext = nullptr;
+		}
+	}
+	
+	virtual void InitRHI() override
+	{
+		FSamplerStateInitializerRHI SamplerStateInitializer(SF_Bilinear, AM_Clamp, AM_Clamp, AM_Clamp);
+		SamplerStateRHI = RHICreateSamplerState(SamplerStateInitializer);
+		
+#if PLATFORM_IOS
+		id<MTLTexture> MetalTexture = Owner->GetMetalTexture();
+		if (MetalTexture)
+		{
+			SCOPED_AUTORELEASE_POOL;
+			
+			CFRetain(MetalTexture);
+			
+			CGColorSpaceRef ColorSpaceRef = CGColorSpaceCreateWithName(kCGColorSpaceGenericRGBLinear);
+			CIImage* Image = [[CIImage alloc] initWithMTLTexture: MetalTexture options: nil];
+
+			// Textures always need to be rotated so to a sane orientation (and mirrored because of differing coord system)
+			CIImage* RotatedImage = [Image imageByApplyingOrientation: GetRotationFromDeviceOrientation()];
+			
+			// Get the sizes from the rotated image
+			CGRect ImageExtent = RotatedImage.extent;
+			
+			FIntPoint DesiredSize(ImageExtent.size.width, ImageExtent.size.height);
+			
+			if (!TextureRHI || DesiredSize != Size)
+			{
+				// Let go of the last texture
+				RHIUpdateTextureReference(Owner->TextureReference.TextureReferenceRHI, nullptr);
+				TextureRHIRef.SafeRelease();
+				
+				Size = DesiredSize;
+				
+				MTLPixelFormat MetalPixelFormat = MetalTexture.pixelFormat;
+				EPixelFormat PixelFormat = EPixelFormat::PF_Unknown;
+				if (MetalPixelFormat == MTLPixelFormatR8Unorm)
+				{
+					PixelFormat = EPixelFormat::PF_G8;
+				}
+				else if (MetalPixelFormat == MTLPixelFormatR16Float)
+				{
+					PixelFormat = EPixelFormat::PF_R16F;
+				}
+				else
+				{
+					UE_LOG(LogTemp, Error, TEXT("FMetalTextureResource::InitRHI: Metal pixel format is not supported: %d"), (int32)MetalPixelFormat);
+				}
+				
+				if (PixelFormat != EPixelFormat::PF_Unknown)
+				{
+					FRHIResourceCreateInfo CreateInfo;
+					TextureRHIRef = RHICreateTexture2D(Size.X, Size.Y, PixelFormat, 1, 1, TexCreate_ShaderResource | TexCreate_UAV, CreateInfo);
+				}
+			}
+			
+			if (TextureRHIRef)
+			{
+				// Get the underlying metal texture so we can render to it
+				id<MTLTexture> UnderlyingMetalTexture = (id<MTLTexture>)TextureRHIRef->GetNativeResource();
+
+				// Do the conversion on the GPU
+				if (!ImageContext)
+				{
+					ImageContext = [CIContext context];
+					CFRetain(ImageContext);
+				}
+				[ImageContext render: RotatedImage toMTLTexture: UnderlyingMetalTexture commandBuffer: nil bounds: ImageExtent colorSpace: ColorSpaceRef];
+			}
+			
+			// Now that the conversion is done, we can get rid of our refs
+			[Image release];
+			CGColorSpaceRelease(ColorSpaceRef);
+			CFRelease(MetalTexture);
+			MetalTexture = nullptr;
+		}
+#endif
+		
+		if (!TextureRHIRef)
+		{
+			// Default to an empty 1x1 texture if we don't have a camera image
+			FRHIResourceCreateInfo CreateInfo;
+			Size.X = Size.Y = 1;
+			TextureRHIRef = RHICreateTexture2D(Size.X, Size.Y, PF_B8G8R8A8, 1, 1, TexCreate_ShaderResource, CreateInfo);
+		}
+
+		TextureRHI = TextureRHIRef;
+		TextureRHI->SetName(Owner->GetFName());
+		RHIBindDebugLabelName(TextureRHI, *Owner->GetName());
+		RHIUpdateTextureReference(Owner->TextureReference.TextureReferenceRHI, TextureRHI);
+	}
+	
+#if PLATFORM_IOS
+	/** @return the rotation to use to rotate the texture to the proper direction */
+	int32 GetRotationFromDeviceOrientation()
+	{
+		// NOTE: The texture we are reading from is in device space and mirrored, because Apple hates us
+		EDeviceScreenOrientation ScreenOrientation = FPlatformMisc::GetDeviceOrientation();
+		switch (ScreenOrientation)
+		{
+			case EDeviceScreenOrientation::Portrait:
+			{
+				return kCGImagePropertyOrientationLeft;
+			}
+
+			case EDeviceScreenOrientation::LandscapeLeft:
+			{
+				return kCGImagePropertyOrientationDown;
+			}
+
+			case EDeviceScreenOrientation::PortraitUpsideDown:
+			{
+				return kCGImagePropertyOrientationUp;
+			}
+
+			case EDeviceScreenOrientation::LandscapeRight:
+			{
+				return kCGImagePropertyOrientationUp;
+			}
+		}
+		
+		// Don't know so don't rotate
+		return kCGImagePropertyOrientationUp;
+	}
+#endif
+	
+	virtual void ReleaseRHI() override
+	{
+		RHIUpdateTextureReference(Owner->TextureReference.TextureReferenceRHI, nullptr);
+		TextureRHIRef.SafeRelease();
+		FTextureResource::ReleaseRHI();
+	}
+	
+	/** Returns the width of the texture in pixels. */
+	virtual uint32 GetSizeX() const override
+	{
+		return Size.X;
+	}
+	
+	/** Returns the height of the texture in pixels. */
+	virtual uint32 GetSizeY() const override
+	{
+		return Size.Y;
+	}
+	
+private:
+	FIntPoint Size;
+	FTexture2DRHIRef TextureRHIRef;
+	const UAppleARKitOcclusionTexture* Owner = nullptr;
+	CIContext* ImageContext = nullptr;
+};
+
+void UAppleARKitOcclusionTexture::SetMetalTexture(float InTimestamp, id<MTLTexture> InMetalTexture)
+{
+	{
+		FScopeLock ScopeLock(&MetalTextureLock);
+		Timestamp = InTimestamp;
+		
+		if (MetalTexture != InMetalTexture)
+		{
+			if (MetalTexture)
+			{
+				CFRelease(MetalTexture);
+				MetalTexture = nullptr;
+			}
+			
+			MetalTexture = InMetalTexture;
+			
+			if (MetalTexture)
+			{
+				CFRetain(MetalTexture);
+				Size = FVector2D(MetalTexture.width, MetalTexture.height);
+			}
+		}
+	}
+	
+	if (Resource == nullptr)
+	{
+		UpdateResource();
+	}
+	
+	if (Resource)
+	{
+		ENQUEUE_RENDER_COMMAND(UpdateMetalTextureResource)
+		([InResource = Resource](FRHICommandListImmediate& RHICmdList)
+		{
+			InResource->InitRHI();
+		});
+	}
+}
+
+id<MTLTexture> UAppleARKitOcclusionTexture::GetMetalTexture() const
+{
+	FScopeLock ScopeLock(&MetalTextureLock);
+	return MetalTexture;
+}
+
+FTextureResource* UAppleARKitOcclusionTexture::CreateResource()
+{
+	return new FOcclusionTextureResource(this);
+}
+
+#else
+
+FTextureResource* UAppleARKitOcclusionTexture::CreateResource()
+{
+	return nullptr;
+}
+
+#endif
diff --git a/Engine/Plugins/Runtime/AR/Apple/AppleARKit/Source/AppleARKit/Private/AppleARKitVideoOverlay.cpp b/Engine/Plugins/Runtime/AR/Apple/AppleARKit/Source/AppleARKit/Private/AppleARKitVideoOverlay.cpp
--- a/Engine/Plugins/Runtime/AR/Apple/AppleARKit/Source/AppleARKit/Private/AppleARKitVideoOverlay.cpp
+++ b/Engine/Plugins/Runtime/AR/Apple/AppleARKit/Source/AppleARKit/Private/AppleARKitVideoOverlay.cpp
@@ -2,6 +2,7 @@
 
 #include "AppleARKitVideoOverlay.h"
 #include "AppleARKitFrame.h"
+#include "AppleARKitSystem.h"
 #include "Materials/Material.h"
 #include "Materials/MaterialInstanceDynamic.h"
 #include "MaterialShaderType.h"
@@ -26,16 +27,45 @@
 	#include "IOSAppDelegate.h"
 #endif
 
+DECLARE_CYCLE_STAT(TEXT("Update Occlusion Textures"), STAT_UpdateOcclusionTextures, STATGROUP_ARKIT);
+
+const FString UARKitCameraOverlayMaterialLoader::OverlayMaterialPath(TEXT("/AppleARKit/M_CameraOverlay.M_CameraOverlay"));
+const FString UARKitCameraOverlayMaterialLoader::DepthOcclusionOverlayMaterialPath(TEXT("/AppleARKit/MI_DepthOcclusionOverlay.MI_DepthOcclusionOverlay"));
+const FString UARKitCameraOverlayMaterialLoader::MatteOcclusionOverlayMaterialPath(TEXT("/AppleARKit/MI_MatteOcclusionOverlay.MI_MatteOcclusionOverlay"));
+
 FAppleARKitVideoOverlay::FAppleARKitVideoOverlay()
 	: MID_CameraOverlay(nullptr)
-	, RenderingOverlayMaterial(nullptr)
 {
-	UMaterialInterface* LoadedMaterial = LoadObject<UMaterialInterface>(nullptr, TEXT("/AppleARKit/M_CameraOverlay.M_CameraOverlay"));
+	UMaterialInterface* LoadedMaterial = LoadObject<UMaterialInterface>(nullptr, *UARKitCameraOverlayMaterialLoader::OverlayMaterialPath);
 	check(LoadedMaterial != nullptr);
 	// Create a MaterialInstanceDynamic which we'll replace the texture with the camera texture from ARKit
 	MID_CameraOverlay = UMaterialInstanceDynamic::Create(LoadedMaterial, GetTransientPackage());
-	RenderingOverlayMaterial = MID_CameraOverlay;
-	check(RenderingOverlayMaterial != nullptr);
+	check(MID_CameraOverlay != nullptr);
+	
+	UMaterialInterface* DepthOcclusionMaterial = LoadObject<UMaterialInterface>(nullptr, *UARKitCameraOverlayMaterialLoader::DepthOcclusionOverlayMaterialPath);
+	check(DepthOcclusionMaterial);
+	MID_DepthOcclusionOverlay = UMaterialInstanceDynamic::Create(DepthOcclusionMaterial, GetTransientPackage());
+	
+	UMaterialInterface* MatteOcclusionMaterial = LoadObject<UMaterialInterface>(nullptr, *UARKitCameraOverlayMaterialLoader::MatteOcclusionOverlayMaterialPath);
+	check(MatteOcclusionMaterial);
+	MID_MatteOcclusionOverlay = UMaterialInstanceDynamic::Create(MatteOcclusionMaterial, GetTransientPackage());
+}
+
+FAppleARKitVideoOverlay::~FAppleARKitVideoOverlay()
+{
+#if SUPPORTS_ARKIT_3_0
+	if (CommandQueue)
+	{
+		[CommandQueue release];
+		CommandQueue = nullptr;
+	}
+	
+	if (MatteGenerator)
+	{
+		[MatteGenerator release];
+		MatteGenerator = nullptr;
+	}
+#endif
 }
 
 // We use something similar to the PostProcessMaterial to render the color camera overlay.
@@ -158,13 +188,65 @@ private:
 IMPLEMENT_MATERIAL_SHADER_TYPE(template<>, FARKitCameraOverlayPS<true>, TEXT("/Engine/Private/PostProcessMaterialShaders.usf"), TEXT("MainPS_ES2"), SF_Pixel);
 IMPLEMENT_MATERIAL_SHADER_TYPE(template<>, FARKitCameraOverlayPS<false>, TEXT("/Engine/Private/PostProcessMaterialShaders.usf"), TEXT("MainPS_VideoOverlay"), SF_Pixel);
 
-void FAppleARKitVideoOverlay::RenderVideoOverlay_RenderThread(FRHICommandListImmediate& RHICmdList, const FSceneView& InView, FAppleARKitFrame& Frame, const EDeviceScreenOrientation DeviceOrientation)
+void FAppleARKitVideoOverlay::UpdateOcclusionTextures(const FAppleARKitFrame& Frame)
+{
+#if SUPPORTS_ARKIT_3_0
+	SCOPE_CYCLE_COUNTER(STAT_UpdateOcclusionTextures);
+	
+	if (FAppleARKitAvailability::SupportsARKit30())
+	{
+		ARFrame* NativeFrame = (ARFrame*)Frame.NativeFrame;
+		if (bEnablePersonOcclusion && NativeFrame &&
+			(NativeFrame.segmentationBuffer || NativeFrame.estimatedDepthData))
+		{
+			id<MTLDevice> Device = (id<MTLDevice>)GDynamicRHI->RHIGetNativeDevice();
+			check(Device);
+			
+			if (!MatteGenerator)
+			{
+				// TODO: add config variable for the matte resolution
+				MatteGenerator = [[ARMatteGenerator alloc] initWithDevice: Device
+														  matteResolution: ARMatteResolutionFull];
+			}
+			
+			if (!CommandQueue)
+			{
+				CommandQueue = [Device newCommandQueue];
+			}
+			
+			id<MTLCommandBuffer> CommandBuffer = [CommandQueue commandBuffer];
+			id<MTLTexture> MatteTexture = nullptr;
+			id<MTLTexture> DepthTexture = nullptr;
+			
+			MatteTexture = [MatteGenerator generateMatteFromFrame: NativeFrame commandBuffer: CommandBuffer];
+			DepthTexture = [MatteGenerator generateDilatedDepthFromFrame: NativeFrame commandBuffer: CommandBuffer];
+			
+			if (MatteTexture && OcclusionMatteTexture)
+			{
+				OcclusionMatteTexture->SetMetalTexture(Frame.Timestamp, MatteTexture);
+			}
+			
+			if (DepthTexture && OcclusionDepthTexture)
+			{
+				OcclusionDepthTexture->SetMetalTexture(Frame.Timestamp, DepthTexture);
+				bOcclusionDepthTextureRecentlyUpdated = true;
+			}
+			
+			[CommandBuffer commit];
+		}
+	}
+#endif
+}
+
+void FAppleARKitVideoOverlay::RenderVideoOverlayWithMaterial(FRHICommandListImmediate& RHICmdList, const FSceneView& InView, struct FAppleARKitFrame& Frame, const EDeviceScreenOrientation DeviceOrientation, UMaterialInstanceDynamic* RenderingOverlayMaterial, const bool bRenderingOcclusion)
 {
 	if (RenderingOverlayMaterial == nullptr || !RenderingOverlayMaterial->IsValidLowLevel())
 	{
 		return;
 	}
-
+	
+	SCOPED_DRAW_EVENTF(RHICmdList, RenderVideoOverlay, bRenderingOcclusion ? TEXT("VideoOverlay (Occlusion)") : TEXT("VideoOverlay (Background)"));
+	
 	if (OverlayVertexBufferRHI[0] == nullptr || !OverlayVertexBufferRHI[0].IsValid())
 	{
 		const FVector2D ViewSize(InView.UnconstrainedViewRect.Max.X, InView.UnconstrainedViewRect.Max.Y);
@@ -266,10 +348,19 @@ void FAppleARKitVideoOverlay::RenderVideoOverlay_RenderThread(FRHICommandListImm
 
 	FGraphicsPipelineStateInitializer GraphicsPSOInit;
 	RHICmdList.ApplyCachedRenderTargets(GraphicsPSOInit);
-
-	GraphicsPSOInit.BlendState = TStaticBlendState<>::GetRHI();
+	
 	GraphicsPSOInit.RasterizerState = TStaticRasterizerState<>::GetRHI();
-	GraphicsPSOInit.DepthStencilState = TStaticDepthStencilState<false, CF_DepthNearOrEqual>::GetRHI();
+	if (bRenderingOcclusion)
+	{
+		GraphicsPSOInit.BlendState = TStaticBlendState<CW_RGBA, BO_Add, BF_SourceAlpha, BF_InverseSourceAlpha, BO_Add, BF_Zero, BF_One>::GetRHI();
+		GraphicsPSOInit.DepthStencilState = TStaticDepthStencilState<false, CF_Always>::GetRHI();
+	}
+	else
+	{
+		// Disable the write mask for the alpha channel so that the scene depth info saved in it is retained
+		GraphicsPSOInit.BlendState = TStaticBlendState<CW_RGB>::GetRHI();
+		GraphicsPSOInit.DepthStencilState = TStaticDepthStencilState<false, CF_DepthNearOrEqual>::GetRHI();
+	}
 	GraphicsPSOInit.PrimitiveType = PT_TriangleList;
 	GraphicsPSOInit.BoundShaderState.VertexDeclarationRHI = GFilterVertexDeclaration.VertexDeclarationRHI;
 
@@ -360,6 +451,34 @@ void FAppleARKitVideoOverlay::RenderVideoOverlay_RenderThread(FRHICommandListImm
 	}
 }
 
+void FAppleARKitVideoOverlay::RenderVideoOverlay_RenderThread(FRHICommandListImmediate& RHICmdList, const FSceneView& InView, FAppleARKitFrame& Frame, const EDeviceScreenOrientation DeviceOrientation, const float WorldToMeterScale)
+{
+	UpdateOcclusionTextures(Frame);
+	
+	RenderVideoOverlayWithMaterial(RHICmdList, InView, Frame, DeviceOrientation, MID_CameraOverlay, false);
+	
+	if (bEnablePersonOcclusion)
+	{
+		UMaterialInstanceDynamic* OcclusionMaterial = bOcclusionDepthTextureRecentlyUpdated ? MID_DepthOcclusionOverlay : MID_MatteOcclusionOverlay;
+		
+		static const FName WorldToMeterScaleParamName(TEXT("WorldToMeterScale"));
+		OcclusionMaterial->SetScalarParameterValue(WorldToMeterScaleParamName, WorldToMeterScale);
+		
+		const bool bIsLandscape = (DeviceOrientation == EDeviceScreenOrientation::LandscapeLeft || DeviceOrientation == EDeviceScreenOrientation::LandscapeRight);
+		FLinearColor DepthTextureUVParamValue = {
+			1.f - (bIsLandscape ? UVOffset.X : UVOffset.Y) * 2.f,
+			1.f - (bIsLandscape ? UVOffset.Y : UVOffset.X) * 2.f,
+			1.f,
+			1.f
+		};
+		
+		static const FName DepthTextureUVParamName(TEXT("DepthTextureUVParam"));
+		OcclusionMaterial->SetVectorParameterValue(DepthTextureUVParamName, DepthTextureUVParamValue);
+		
+		RenderVideoOverlayWithMaterial(RHICmdList, InView, Frame, DeviceOrientation, OcclusionMaterial, true);
+		bOcclusionDepthTextureRecentlyUpdated = false;
+	}
+}
 
 bool FAppleARKitVideoOverlay::GetPassthroughCameraUVs_RenderThread(TArray<FVector2D>& OutUVs, const EDeviceScreenOrientation DeviceOrientation)
 {
@@ -411,11 +530,56 @@ bool FAppleARKitVideoOverlay::GetPassthroughCameraUVs_RenderThread(TArray<FVecto
 
 void FAppleARKitVideoOverlay::AddReferencedObjects(FReferenceCollector& Collector)
 {
-	Collector.AddReferencedObject(RenderingOverlayMaterial);
 	Collector.AddReferencedObject(MID_CameraOverlay);
+	Collector.AddReferencedObject(OcclusionMatteTexture);
+	Collector.AddReferencedObject(OcclusionDepthTexture);
+	Collector.AddReferencedObject(MID_DepthOcclusionOverlay);
+	Collector.AddReferencedObject(MID_MatteOcclusionOverlay);
 }
 
 void FAppleARKitVideoOverlay::SetOverlayTexture(UARTextureCameraImage* InCameraImage)
 {
-	MID_CameraOverlay->SetTextureParameterValue(FName("CameraImage"), InCameraImage);
+	static const FName ParamName(TEXT("CameraImage"));
+	MID_CameraOverlay->SetTextureParameterValue(ParamName, InCameraImage);
+	MID_DepthOcclusionOverlay->SetTextureParameterValue(ParamName, InCameraImage);
+	MID_MatteOcclusionOverlay->SetTextureParameterValue(ParamName, InCameraImage);
+}
+
+void FAppleARKitVideoOverlay::SetEnablePersonOcclusion(bool bEnable)
+{
+#if SUPPORTS_ARKIT_3_0
+	bEnablePersonOcclusion = bEnable;
+	
+	if (bEnablePersonOcclusion)
+	{
+		// TODO: add new EARTextureType for the occlusion textures
+		OcclusionMatteTexture = NewObject<UAppleARKitOcclusionTexture>();
+		OcclusionMatteTexture->TextureType = EARTextureType::CameraImage;
+		OcclusionMatteTexture->UpdateResource();
+		
+		OcclusionDepthTexture = NewObject<UAppleARKitOcclusionTexture>();
+		OcclusionDepthTexture->TextureType = EARTextureType::CameraImage;
+		OcclusionDepthTexture->UpdateResource();
+	}
+	else
+	{
+		OcclusionMatteTexture = OcclusionDepthTexture = nullptr;
+	}
+	
+	static const FName MatteTextureParamName(TEXT("OcclusionMatteTexture"));
+	static const FName DepthTextureParamName(TEXT("OcclusionDepthTexture"));
+	
+	MID_DepthOcclusionOverlay->SetTextureParameterValue(MatteTextureParamName, OcclusionMatteTexture);
+	MID_DepthOcclusionOverlay->SetTextureParameterValue(DepthTextureParamName, OcclusionDepthTexture);
+	
+	MID_MatteOcclusionOverlay->SetTextureParameterValue(MatteTextureParamName, OcclusionMatteTexture);
+	MID_MatteOcclusionOverlay->SetTextureParameterValue(DepthTextureParamName, OcclusionDepthTexture);
+	
+	// we need to clear the scene color with max alpha to ensure the scene depth is correct on mobile
+	static IConsoleVariable* CVarSceneColorClearAlphaFlag = IConsoleManager::Get().FindConsoleVariable(TEXT("r.MobileClearSceneColorWithMaxAlpha"));
+	if (CVarSceneColorClearAlphaFlag)
+	{
+		CVarSceneColorClearAlphaFlag->Set(bEnablePersonOcclusion ? 1 : 0);
+	}
+#endif
 }
diff --git a/Engine/Plugins/Runtime/AR/Apple/AppleARKit/Source/AppleARKit/Private/AppleARKitVideoOverlay.h b/Engine/Plugins/Runtime/AR/Apple/AppleARKit/Source/AppleARKit/Private/AppleARKitVideoOverlay.h
--- a/Engine/Plugins/Runtime/AR/Apple/AppleARKit/Source/AppleARKit/Private/AppleARKitVideoOverlay.h
+++ b/Engine/Plugins/Runtime/AR/Apple/AppleARKit/Source/AppleARKit/Private/AppleARKitVideoOverlay.h
@@ -5,10 +5,12 @@
 #include "CoreMinimal.h"
 #include "Materials/MaterialInterface.h"
 #include "UObject/ConstructorHelpers.h"
+#include "AppleARKitAvailability.h"
 #include "AppleARKitVideoOverlay.generated.h"
 
 class UARTextureCameraImage;
 class UMaterialInstanceDynamic;
+class UAppleARKitOcclusionTexture;
 
 /** Helper class to ensure the ARKit camera material is cooked. */
 UCLASS()
@@ -19,12 +21,28 @@ class UARKitCameraOverlayMaterialLoader : public UObject
 public:
 	UPROPERTY()
 	UMaterialInterface* DefaultCameraOverlayMaterial;
+	
+	UPROPERTY()
+	UMaterialInterface* DepthOcclusionOverlayMaterial;
+	
+	UPROPERTY()
+	UMaterialInterface* MatteOcclusionOverlayMaterial;
 
 	UARKitCameraOverlayMaterialLoader()
 	{
-		static ConstructorHelpers::FObjectFinder<UMaterialInterface> DefaultOverlayMaterialRef(TEXT("/AppleARKit/M_CameraOverlay.M_CameraOverlay"));
+		static ConstructorHelpers::FObjectFinder<UMaterialInterface> DefaultOverlayMaterialRef(*OverlayMaterialPath);
 		DefaultCameraOverlayMaterial = DefaultOverlayMaterialRef.Object;
+		
+		static ConstructorHelpers::FObjectFinder<UMaterialInterface> DepthOcclusionOverlayMaterialRef(*DepthOcclusionOverlayMaterialPath);
+		DepthOcclusionOverlayMaterial = DepthOcclusionOverlayMaterialRef.Object;
+		
+		static ConstructorHelpers::FObjectFinder<UMaterialInterface> MatteOcclusionOverlayMaterialRef(*MatteOcclusionOverlayMaterialPath);
+		MatteOcclusionOverlayMaterial = MatteOcclusionOverlayMaterialRef.Object;
 	}
+	
+	static const FString OverlayMaterialPath;
+	static const FString DepthOcclusionOverlayMaterialPath;
+	static const FString MatteOcclusionOverlayMaterialPath;
 };
 
 class FAppleARKitVideoOverlay
@@ -32,25 +50,42 @@ class FAppleARKitVideoOverlay
 {
 public:
 	FAppleARKitVideoOverlay();
+	virtual ~FAppleARKitVideoOverlay();
 
 	void SetCameraTexture(UARTextureCameraImage* InCameraImage);
 
-	void RenderVideoOverlay_RenderThread(FRHICommandListImmediate& RHICmdList, const FSceneView& InView, struct FAppleARKitFrame& Frame, const EDeviceScreenOrientation DeviceOrientation);
+	void RenderVideoOverlay_RenderThread(FRHICommandListImmediate& RHICmdList, const FSceneView& InView, struct FAppleARKitFrame& Frame, const EDeviceScreenOrientation DeviceOrientation, const float WorldToMeterScale);
 	bool GetPassthroughCameraUVs_RenderThread(TArray<FVector2D>& OutUVs, const EDeviceScreenOrientation DeviceOrientation);
 
 	void SetOverlayTexture(UARTextureCameraImage* InCameraImage);
+	void SetEnablePersonOcclusion(bool bEnable);
 
 private:
 	//~ FGCObject
 	virtual void AddReferencedObjects(FReferenceCollector& Collector) override;
 	//~ FGCObject
+	
+	void RenderVideoOverlayWithMaterial(FRHICommandListImmediate& RHICmdList, const FSceneView& InView, struct FAppleARKitFrame& Frame, const EDeviceScreenOrientation DeviceOrientation, UMaterialInstanceDynamic* RenderingOverlayMaterial, const bool bRenderingOcclusion);
+	void UpdateOcclusionTextures(const FAppleARKitFrame& Frame);
 
 	UMaterialInstanceDynamic* MID_CameraOverlay;
-	UMaterialInterface* RenderingOverlayMaterial;
 
 	// Cache UV offset to be used by GetPassthroughCameraUVs_RenderThread
 	FVector2D UVOffset;
 
 	FVertexBufferRHIRef OverlayVertexBufferRHI[2];
 	FIndexBufferRHIRef IndexBufferRHI;
+	
+	bool bEnablePersonOcclusion = false;
+	
+#if SUPPORTS_ARKIT_3_0
+	ARMatteGenerator* MatteGenerator = nullptr;
+	id<MTLCommandQueue> CommandQueue = nullptr;
+#endif
+	
+	bool bOcclusionDepthTextureRecentlyUpdated = false;
+	UAppleARKitOcclusionTexture* OcclusionMatteTexture = nullptr;
+	UAppleARKitOcclusionTexture* OcclusionDepthTexture = nullptr;
+	UMaterialInstanceDynamic* MID_DepthOcclusionOverlay = nullptr;
+	UMaterialInstanceDynamic* MID_MatteOcclusionOverlay = nullptr;
 };
diff --git a/Engine/Plugins/Runtime/AR/Apple/AppleARKit/Source/AppleARKit/Public/AppleARKitAvailability.h b/Engine/Plugins/Runtime/AR/Apple/AppleARKit/Source/AppleARKit/Public/AppleARKitAvailability.h
--- a/Engine/Plugins/Runtime/AR/Apple/AppleARKit/Source/AppleARKit/Public/AppleARKitAvailability.h
+++ b/Engine/Plugins/Runtime/AR/Apple/AppleARKit/Source/AppleARKit/Public/AppleARKitAvailability.h
@@ -31,13 +31,20 @@
 		#define SUPPORTS_ARKIT_2_0 0
 	#endif
 
+	// Check for ARKit 3.0
+	#ifdef __IPHONE_13_0
+		#define SUPPORTS_ARKIT_3_0 1
+	#else
+		#define SUPPORTS_ARKIT_3_0 0
+	#endif
+
 #else
 
 	// No ARKit support
 	#define SUPPORTS_ARKIT_1_0 0
 	#define SUPPORTS_ARKIT_1_5 0
 	#define SUPPORTS_ARKIT_2_0 0
-
+	#define SUPPORTS_ARKIT_3_0 0
 #endif
 
 class FAppleARKitAvailability
@@ -96,4 +103,22 @@ public:
 #endif
 		return bSupportsARKit20;
 	}
+	
+	static bool SupportsARKit30()
+	{
+		static bool bSupportsARKit30 = false;
+#if SUPPORTS_ARKIT_3_0
+		static bool bSupportChecked = false;
+		if (!bSupportChecked)
+		{
+			bSupportChecked = true;
+			// This call is slow, so cache the result
+			if (@available(iOS 13.0, *))
+			{
+				bSupportsARKit30 = true;
+			}
+		}
+#endif
+		return bSupportsARKit30;
+	}
 };
diff --git a/Engine/Plugins/Runtime/AR/Apple/AppleARKit/Source/AppleARKit/Public/AppleARKitConversion.h b/Engine/Plugins/Runtime/AR/Apple/AppleARKit/Source/AppleARKit/Public/AppleARKitConversion.h
--- a/Engine/Plugins/Runtime/AR/Apple/AppleARKit/Source/AppleARKit/Public/AppleARKitConversion.h
+++ b/Engine/Plugins/Runtime/AR/Apple/AppleARKit/Source/AppleARKit/Public/AppleARKitConversion.h
@@ -248,6 +248,24 @@ struct APPLEARKIT_API FAppleARKitConversion
 
 	#pragma clang diagnostic pop
 #endif
+	
+#if SUPPORTS_ARKIT_1_0
+	// Helper function to check if a particular session feature is supported with the specified session type
+	static bool IsSessionTrackingFeatureSupported(EARSessionType SessionType, EARSessionTrackingFeature SessionTrackingFeature);
+#else
+	static bool IsSessionTrackingFeatureSupported(EARSessionType SessionType, EARSessionTrackingFeature SessionTrackingFeature) { return false; }
+#endif
+		
+#if SUPPORTS_ARKIT_3_0
+	static void InitImageDetection(UARSessionConfig* SessionConfig, ARBodyTrackingConfiguration* BodyTrackingConfig, TMap< FString, UARCandidateImage* >& CandidateImages, TMap< FString, CGImageRef >& ConvertedCandidateImages);
+	static ARFrameSemantics ToARFrameSemantics(EARSessionTrackingFeature SessionTrackingFeature);
+	static FARPose2D ToARPose2D(const ARBody2D* InARPose2D);
+	static FARPose3D ToARPose3D(const ARBodyAnchor* InARBodyAnchor);
+    static FARPose3D ToARPose3D(const ARSkeleton3D* InSkeleton3D, bool bIdentityForUntracked);
+
+private:
+	static void ToSkeletonDefinition(const ARSkeletonDefinition* InARSkeleton, FARSkeletonDefinition& OutSkeletonDefinition);
+#endif
 };
 
 #if SUPPORTS_ARKIT_1_0
@@ -259,6 +277,7 @@ enum class EAppleAnchorType : uint8
 	ImageAnchor,
 	EnvironmentProbeAnchor,
 	ObjectAnchor,
+	PoseAnchor,
 	MAX
 };
 
@@ -328,6 +347,14 @@ struct FAppleARKitAnchorData
 	{
 	}
 
+	FAppleARKitAnchorData(FGuid InAnchorGuid, FTransform InTransform, FARPose3D InTrackedPose)
+		: Transform( InTransform )
+		, AnchorType( EAppleAnchorType::PoseAnchor )
+		, AnchorGUID( InAnchorGuid )
+		, TrackedPose( MoveTemp(InTrackedPose) )
+	{
+	}
+	
 	FAppleARKitAnchorData& operator=(const FAppleARKitAnchorData& Other)
 	{
 		if (this != &Other)
@@ -369,6 +396,8 @@ struct FAppleARKitAnchorData
 		ObjectClassification = Other.ObjectClassification;
 
 		bIsTracked = Other.bIsTracked;
+		
+		TrackedPose = Other.TrackedPose;
 	}
 
 	void Clear()
@@ -379,6 +408,8 @@ struct FAppleARKitAnchorData
 		Vertices.Empty();
 		Indices.Empty();
 		ProbeTexture = nullptr;
+		
+		TrackedPose = {};
 	}
 
 	FTransform Transform;
@@ -413,6 +444,12 @@ struct FAppleARKitAnchorData
 
 	/** Only valid for tracked real world objects (face, images) */
 	bool bIsTracked;
+	
+	/** Only valid if this is a body anchor */
+	FARPose3D TrackedPose;
+
+	// only need to be initialized once
+	static TSharedPtr<FARPose3D> BodyRefPose;
 };
 #endif
 
diff --git a/Engine/Plugins/Runtime/AR/Apple/AppleARKit/Source/AppleARKit/Public/AppleARKitFrame.h b/Engine/Plugins/Runtime/AR/Apple/AppleARKit/Source/AppleARKit/Public/AppleARKitFrame.h
--- a/Engine/Plugins/Runtime/AR/Apple/AppleARKit/Source/AppleARKit/Public/AppleARKitFrame.h
+++ b/Engine/Plugins/Runtime/AR/Apple/AppleARKit/Source/AppleARKit/Public/AppleARKitFrame.h
@@ -28,8 +28,10 @@ struct APPLEARKIT_API FAppleARKitFrame
 	 * with the UE4-ified versions of ARFrames properties.
 	 *
 	 * @param InARFrame 		- The frame to convert / copy
+	 * @param MinCameraUV 		- The minimum (top left) UV used to render the passthrough camera
+	 * @param MaxCameraUV 		- The maximum (bottom right) UV used to render the passthrough camera
 	 */
-	FAppleARKitFrame( ARFrame* InARFrame );
+	FAppleARKitFrame( ARFrame* InARFrame, const FVector2D MinCameraUV, const FVector2D MaxCameraUV );
 
 	/** 
 	 * Copy contructor. CapturedImage is skipped as we don't need / want to retain access to the 
@@ -69,6 +71,17 @@ struct APPLEARKIT_API FAppleARKitFrame
 	
 	/** The current world mapping state is reported on the frame */
 	EARWorldMappingState WorldMappingState;
+	
+	/** The current tracked 2D pose */
+	FARPose2D Tracked2DPose;
+
+#if SUPPORTS_ARKIT_3_0
+	/** The person segmentation buffer from ARKit */
+	CVPixelBufferRef SegmentationBuffer = nullptr;
+		
+	/** The estimated depth buffer for person segmentation from ARKit */
+	CVPixelBufferRef EstimatedDepthData = nullptr;
+#endif
 
 	/* 
 	 * When adding new member variables, don't forget to handle them in the copy constructor and
diff --git a/Engine/Plugins/Runtime/AR/Apple/AppleARKit/Source/AppleARKit/Public/AppleARKitPoseTrackingLiveLink.h b/Engine/Plugins/Runtime/AR/Apple/AppleARKit/Source/AppleARKit/Public/AppleARKitPoseTrackingLiveLink.h
new file mode 100644
--- /dev/null
+++ b/Engine/Plugins/Runtime/AR/Apple/AppleARKit/Source/AppleARKit/Public/AppleARKitPoseTrackingLiveLink.h
@@ -0,0 +1,36 @@
+// Copyright 1998-2019 Epic Games, Inc. All Rights Reserved.
+
+#pragma once
+
+#include "Features/IModularFeature.h"
+#include "Features/IModularFeatures.h"
+
+#include "AppleARKitAvailability.h"
+
+#if SUPPORTS_ARKIT_3_0
+	#import <ARKit/ARKit.h>
+#endif
+
+struct FAppleARKitAnchorData;
+
+class APPLEARKIT_API IAppleARKitPoseTrackingLiveLink
+	: public IModularFeature
+{
+public:
+	/**
+	 * Publishes any body tracking AR data that needs to be sent to LiveLink. Done as a separate step because MakeAnchorData is called
+	 * on an arbitrary thread and we can't access UObjects there safely
+	 *
+	 * @param AnchorList the list of anchors to publish to LiveLink
+	 *
+	 * @return the set of body tracking anchors to dispatch
+	 */
+	virtual void PublishLiveLinkData(TSharedPtr<FAppleARKitAnchorData> Anchor) { }
+
+	virtual const TArray<FTransform>* GetRefPoseTransforms() const { return nullptr; }
+	static FName GetModularFeatureName()
+	{
+		static FName FeatureName = FName(TEXT("AppleARKitPoseTrackingLiveLink"));
+		return FeatureName;
+	}
+};
diff --git a/Engine/Plugins/Runtime/AR/Apple/AppleARKit/Source/AppleARKit/Public/AppleARKitSettings.h b/Engine/Plugins/Runtime/AR/Apple/AppleARKit/Source/AppleARKit/Public/AppleARKitSettings.h
--- a/Engine/Plugins/Runtime/AR/Apple/AppleARKit/Source/AppleARKit/Public/AppleARKitSettings.h
+++ b/Engine/Plugins/Runtime/AR/Apple/AppleARKit/Source/AppleARKit/Public/AppleARKitSettings.h
@@ -23,6 +23,14 @@ enum class EARFaceTrackingFileWriterType : uint8
 	JSON
 };
 
+UENUM(BlueprintType, Category="AR AugmentedReality")
+enum class ELivelinkTrackingType : uint8
+{
+	None,
+	FaceTracking,
+	PoseTracking
+};
+
 
 UCLASS(Config=Engine, defaultconfig)
 class APPLEARKIT_API UAppleARKitSettings :
@@ -33,7 +41,7 @@ class APPLEARKIT_API UAppleARKitSettings :
 
 public:
 	UAppleARKitSettings()
-		: bEnableLiveLinkForFaceTracking(false)
+		: LivelinkTrackingType(ELivelinkTrackingType::None)
 		, bFaceTrackingLogData(false)
 		, bFaceTrackingWriteEachFrame(false)
 		, FaceTrackingFileWriterType(EARFaceTrackingFileWriterType::None)
@@ -58,6 +66,7 @@ public:
 // Accessors to the properties
 	FString GetFaceTrackingLogDir();
 	bool IsLiveLinkEnabledForFaceTracking();
+	bool IsLiveLinkEnabledForPoseTracking();
 	bool IsFaceTrackingLoggingEnabled();
 	bool ShouldFaceTrackingLogPerFrame();
 	EARFaceTrackingFileWriterType GetFaceTrackingFileWriterType();
@@ -81,9 +90,9 @@ protected:
 	UPROPERTY(Config, BlueprintReadOnly, EditAnywhere, Category="AR Settings")
 	FString FaceTrackingLogDir;
 
-	/** Whether to publish face blend shapes to LiveLink or not */
+	/** Livelink tracking type. To publish face blend shapes, or body pose data to LiveLink, or none */
 	UPROPERTY(Config, BlueprintReadOnly, EditAnywhere, Category="AR Settings")
-	bool bEnableLiveLinkForFaceTracking;
+	ELivelinkTrackingType LivelinkTrackingType;
 
 	/** Whether file writing is enabled at all or not */
 	UPROPERTY(Config, BlueprintReadOnly, EditAnywhere, Category="AR Settings")
diff --git a/Engine/Plugins/Runtime/AR/Apple/AppleARKit/Source/AppleARKit/Public/AppleARKitTextures.h b/Engine/Plugins/Runtime/AR/Apple/AppleARKit/Source/AppleARKit/Public/AppleARKitTextures.h
--- a/Engine/Plugins/Runtime/AR/Apple/AppleARKit/Source/AppleARKit/Public/AppleARKitTextures.h
+++ b/Engine/Plugins/Runtime/AR/Apple/AppleARKit/Source/AppleARKit/Public/AppleARKitTextures.h
@@ -125,4 +125,34 @@ private:
 #endif
 };
 
-
+UCLASS(BlueprintType)
+class APPLEARKIT_API UAppleARKitOcclusionTexture :
+	public UARTexture,
+	public IAppleImageInterface
+{
+	GENERATED_UCLASS_BODY()
+	
+public:
+	// UTexture interface implementation
+	virtual void BeginDestroy() override;
+	virtual FTextureResource* CreateResource() override;
+	virtual EMaterialValueType GetMaterialType() const override { return MCT_Texture2D; }
+	virtual float GetSurfaceWidth() const override { return Size.X; }
+	virtual float GetSurfaceHeight() const override { return Size.Y; }
+	// End UTexture interface
+	
+	virtual EAppleTextureType GetTextureType() const override { return EAppleTextureType::MetalTexture; }
+	
+#if PLATFORM_MAC || PLATFORM_IOS
+	void SetMetalTexture(float InTimestamp, id<MTLTexture> InMetalTexture);
+	
+	// IAppleImageInterface interface implementation
+	virtual id<MTLTexture> GetMetalTexture() const override;
+	// End IAppleImageInterface interface
+	
+private:
+	/** The Apple specific representation of the ar camera image */
+	id<MTLTexture> MetalTexture = nullptr;
+	mutable FCriticalSection MetalTextureLock;
+#endif
+};
diff --git a/Engine/Plugins/Runtime/AR/Apple/AppleARKit/Source/AppleARKitPoseTrackingLiveLink/AppleARKitPoseTrackingLiveLink.build.cs b/Engine/Plugins/Runtime/AR/Apple/AppleARKit/Source/AppleARKitPoseTrackingLiveLink/AppleARKitPoseTrackingLiveLink.build.cs
new file mode 100644
--- /dev/null
+++ b/Engine/Plugins/Runtime/AR/Apple/AppleARKit/Source/AppleARKitPoseTrackingLiveLink/AppleARKitPoseTrackingLiveLink.build.cs
@@ -0,0 +1,55 @@
+// Copyright 1998-2019 Epic Games, Inc. All Rights Reserved.
+
+using UnrealBuildTool;
+
+public class AppleARKitPoseTrackingLiveLink : ModuleRules
+{
+	public AppleARKitPoseTrackingLiveLink(ReadOnlyTargetRules Target) : base(Target)
+	{
+        //OptimizeCode = CodeOptimization.Never;
+        
+		PrivateIncludePaths.AddRange(
+			new string[] {
+				"../../../../../../Plugins/Runtime/AR/Apple/AppleARKit/Source/AppleARKit/Private",
+				// ... add other private include paths required here ...
+			}
+			);
+			
+		
+		PublicDependencyModuleNames.AddRange(
+			new string[]
+			{
+				"Core",
+				"Engine",
+				// ... add other public dependencies that you statically link with here ...
+			}
+			);
+			
+		
+		PrivateDependencyModuleNames.AddRange(
+			new string[]
+			{
+				"CoreUObject",
+                "HeadMountedDisplay",
+//                "AugmentedReality",
+                "LiveLink",
+                "LiveLinkInterface",
+                "AppleARKit",
+                "AppleImageUtils"
+			}
+			);
+		
+		
+		DynamicallyLoadedModuleNames.AddRange(
+			new string[]
+			{
+				// ... add any modules that your module loads dynamically here ...
+			}
+			);
+
+		if (Target.Platform == UnrealTargetPlatform.IOS)
+		{
+			PublicFrameworks.Add( "ARKit" );
+		}
+	}
+}
diff --git a/Engine/Plugins/Runtime/AR/Apple/AppleARKit/Source/AppleARKitPoseTrackingLiveLink/Private/AppleARKitPoseTrackingLiveLinkImpl.cpp b/Engine/Plugins/Runtime/AR/Apple/AppleARKit/Source/AppleARKitPoseTrackingLiveLink/Private/AppleARKitPoseTrackingLiveLinkImpl.cpp
new file mode 100644
--- /dev/null
+++ b/Engine/Plugins/Runtime/AR/Apple/AppleARKit/Source/AppleARKitPoseTrackingLiveLink/Private/AppleARKitPoseTrackingLiveLinkImpl.cpp
@@ -0,0 +1,60 @@
+// Copyright 1998-2019 Epic Games, Inc. All Rights Reserved.
+
+#include "AppleARKitPoseTrackingLiveLinkImpl.h"
+#include "AppleARKitSettings.h"
+#include "AppleARKitConversion.h"
+#include "ARSystem.h"
+#include "AppleARKitPoseTrackingLiveLinkModule.h"
+
+FAppleARKitPoseTrackingLiveLink::FAppleARKitPoseTrackingLiveLink() :
+	LocalDeviceId(FName(*FPlatformMisc::GetDeviceId()))	// Generate our device id
+{
+}
+
+FAppleARKitPoseTrackingLiveLink::~FAppleARKitPoseTrackingLiveLink()
+{
+	// Should only be called durirng shutdown
+	check(GIsRequestingExit);
+}
+
+void FAppleARKitPoseTrackingLiveLink::Init()
+{
+	IModularFeatures::Get().RegisterModularFeature(GetModularFeatureName(), this);
+}
+
+void FAppleARKitPoseTrackingLiveLink::Shutdown()
+{
+	IModularFeatures::Get().UnregisterModularFeature(GetModularFeatureName(), this);
+}
+
+#if SUPPORTS_ARKIT_3_0
+void FAppleARKitPoseTrackingLiveLink::PublishLiveLinkData(TSharedPtr<FAppleARKitAnchorData> Anchor)
+{
+	static bool bNeedsInit = true;
+	if (bNeedsInit)
+	{
+		bNeedsInit = false;
+	    // Create our LiveLink provider if the project setting is enabled
+		if (GetMutableDefault<UAppleARKitSettings>()->IsLiveLinkEnabledForPoseTracking())
+		{
+			LiveLinkSource = FAppleARKitPoseTrackingLiveLinkSourceFactory::CreateLiveLinkSource();
+		}
+	}
+
+	if (LiveLinkSource.IsValid())
+	{
+		PoseTrackingLiveLinkSubjectName = GetMutableDefault<UAppleARKitSettings>()->GetLiveLinkSubjectName();
+        LiveLinkSource->UpdateLiveLink(PoseTrackingLiveLinkSubjectName, Anchor->Timecode, Anchor->FrameRate, Anchor->TrackedPose, *FAppleARKitAnchorData::BodyRefPose, LocalDeviceId);
+	}
+}
+#endif
+
+const TArray<FTransform>* FAppleARKitPoseTrackingLiveLink::GetRefPoseTransforms() const
+{
+	if (LiveLinkSource.IsValid())
+	{
+		return LiveLinkSource->GetRefPoseTransforms();
+	}
+
+	return nullptr;
+}
diff --git a/Engine/Plugins/Runtime/AR/Apple/AppleARKit/Source/AppleARKitPoseTrackingLiveLink/Private/AppleARKitPoseTrackingLiveLinkImpl.h b/Engine/Plugins/Runtime/AR/Apple/AppleARKit/Source/AppleARKitPoseTrackingLiveLink/Private/AppleARKitPoseTrackingLiveLinkImpl.h
new file mode 100644
--- /dev/null
+++ b/Engine/Plugins/Runtime/AR/Apple/AppleARKit/Source/AppleARKitPoseTrackingLiveLink/Private/AppleARKitPoseTrackingLiveLinkImpl.h
@@ -0,0 +1,41 @@
+// Copyright 1998-2019 Epic Games, Inc. All Rights Reserved.
+
+#pragma once
+
+#include "CoreMinimal.h"
+#include "AppleARKitPoseTrackingLiveLink.h"
+#include "AppleARKitPoseTrackingLiveLinkSourceFactory.h"
+#include "AppleARKitConversion.h"
+
+
+class APPLEARKITPOSETRACKINGLIVELINK_API FAppleARKitPoseTrackingLiveLink :
+	public IAppleARKitPoseTrackingLiveLink,
+	public TSharedFromThis<FAppleARKitPoseTrackingLiveLink, ESPMode::ThreadSafe>
+{
+public:
+	FAppleARKitPoseTrackingLiveLink();
+	virtual ~FAppleARKitPoseTrackingLiveLink();
+
+	void Init();
+	void Shutdown();
+
+private:
+#if SUPPORTS_ARKIT_3_0
+	// ~IAppleARKitPoseTrackingLiveLink
+	void PublishLiveLinkData(TSharedPtr<FAppleARKitAnchorData> Anchor) override;
+	// ~IAppleARKitPoseTrackingLiveLink
+
+#endif
+
+	// ~IAppleARKitPoseTrackingLiveLink
+	const TArray<FTransform>* GetRefPoseTransforms() const override;
+	// ~IAppleARKitPoseTrackingLiveLink	
+
+	/** If requested, publishes face ar updates to LiveLink for the animation system to use */
+	TSharedPtr<ILiveLinkSourceARKitPoseTracking> LiveLinkSource;
+	/** Copied from the UARSessionConfig project settings object */
+	FName PoseTrackingLiveLinkSubjectName;
+	/** The id of this device */
+	FName LocalDeviceId;
+};
+
diff --git a/Engine/Plugins/Runtime/AR/Apple/AppleARKit/Source/AppleARKitPoseTrackingLiveLink/Private/AppleARKitPoseTrackingLiveLinkModule.cpp b/Engine/Plugins/Runtime/AR/Apple/AppleARKit/Source/AppleARKitPoseTrackingLiveLink/Private/AppleARKitPoseTrackingLiveLinkModule.cpp
new file mode 100644
--- /dev/null
+++ b/Engine/Plugins/Runtime/AR/Apple/AppleARKit/Source/AppleARKitPoseTrackingLiveLink/Private/AppleARKitPoseTrackingLiveLinkModule.cpp
@@ -0,0 +1,34 @@
+// Copyright 1998-2019 Epic Games, Inc. All Rights Reserved.
+
+#include "AppleARKitPoseTrackingLiveLinkModule.h"
+#include "AppleARKitSystem.h"
+#include "Features/IModularFeature.h"
+#include "Features/IModularFeatures.h"
+#include "CoreGlobals.h"
+#include "AppleARKitPoseTrackingLiveLinkImpl.h"
+
+
+IMPLEMENT_MODULE(FAppleARKitPoseTrackingLiveLinkModule, AppleARKitPoseTrackingLiveLink);
+
+DEFINE_LOG_CATEGORY(LogAppleARKitPoseTracking);
+
+TSharedPtr<FAppleARKitPoseTrackingLiveLink, ESPMode::ThreadSafe> PoseTrackingLiveLinkInstance;
+
+void FAppleARKitPoseTrackingLiveLinkModule::StartupModule()
+{
+	ensureMsgf(FModuleManager::Get().LoadModule("AppleARKit"), TEXT("ARKitPoseTrackingLiveLink depends on the AppleARKit module."));
+
+	PoseTrackingLiveLinkInstance = MakeShared<FAppleARKitPoseTrackingLiveLink, ESPMode::ThreadSafe>();
+	PoseTrackingLiveLinkInstance->Init();
+
+	FQuat AppleForward = FVector(-1, 0, 0).ToOrientationQuat();
+	FQuat MeshForward = FVector(0, 1, 0).ToOrientationQuat();;
+	const auto AppleSpaceToMeshSpace = FTransform(AppleForward).Inverse() * FTransform(MeshForward);
+}
+
+void FAppleARKitPoseTrackingLiveLinkModule::ShutdownModule()
+{
+	PoseTrackingLiveLinkInstance->Shutdown();
+	PoseTrackingLiveLinkInstance = nullptr;
+}
+
diff --git a/Engine/Plugins/Runtime/AR/Apple/AppleARKit/Source/AppleARKitPoseTrackingLiveLink/Private/AppleARKitPoseTrackingLiveLinkModule.h b/Engine/Plugins/Runtime/AR/Apple/AppleARKit/Source/AppleARKitPoseTrackingLiveLink/Private/AppleARKitPoseTrackingLiveLinkModule.h
new file mode 100644
--- /dev/null
+++ b/Engine/Plugins/Runtime/AR/Apple/AppleARKit/Source/AppleARKitPoseTrackingLiveLink/Private/AppleARKitPoseTrackingLiveLinkModule.h
@@ -0,0 +1,19 @@
+// Copyright 1998-2019 Epic Games, Inc. All Rights Reserved.
+
+#pragma once
+
+#include "Stats/Stats.h"
+#include "Modules/ModuleManager.h"
+
+class APPLEARKITPOSETRACKINGLIVELINK_API FAppleARKitPoseTrackingLiveLinkModule :
+	public IModuleInterface
+{
+public:
+	virtual void StartupModule() override;
+
+	virtual void ShutdownModule() override;
+};
+
+DECLARE_LOG_CATEGORY_EXTERN(LogAppleARKitPoseTracking, Log, All);
+
+DECLARE_STATS_GROUP(TEXT("PoseTracking AR"), STATGROUP_PoseTrackingAR, STATCAT_Advanced);
diff --git a/Engine/Plugins/Runtime/AR/Apple/AppleARKit/Source/AppleARKitPoseTrackingLiveLink/Private/AppleARKitPoseTrackingLiveLinkRemapAsset.cpp b/Engine/Plugins/Runtime/AR/Apple/AppleARKit/Source/AppleARKitPoseTrackingLiveLink/Private/AppleARKitPoseTrackingLiveLinkRemapAsset.cpp
new file mode 100644
--- /dev/null
+++ b/Engine/Plugins/Runtime/AR/Apple/AppleARKit/Source/AppleARKitPoseTrackingLiveLink/Private/AppleARKitPoseTrackingLiveLinkRemapAsset.cpp
@@ -0,0 +1,208 @@
+// Copyright 1998-2019 Epic Games, Inc. All Rights Reserved.
+
+#include "AppleARKitPoseTrackingLiveLinkRemapAsset.h"
+#include "Roles/LiveLinkAnimationTypes.h"
+#include "Features/IModularFeatures.h"
+#include "AppleARKitPoseTrackingLiveLink.h"
+
+UAppleARKitPoseTrackingLiveLinkRemapAsset::UAppleARKitPoseTrackingLiveLinkRemapAsset(const FObjectInitializer& ObjectInitializer)
+: Super(ObjectInitializer)
+{
+	// default to unreal tutorial mesh rigging
+	AppleARKitBoneNamesToMeshBoneNames = 
+	{
+		{TEXT("hips_joint"), 					TEXT("pelvis")},
+		{TEXT("spine_2_joint"), 				TEXT("spine_01")},
+		{TEXT("spine_4_joint"), 				TEXT("spine_02")},
+		{TEXT("spine_7_joint"), 				TEXT("spine_03")},
+		{TEXT("left_shoulder_1_joint"), 		TEXT("clavicle_l")},
+		{TEXT("left_arm_joint"), 				TEXT("upperarm_l")},
+		{TEXT("left_forearm_joint"), 			TEXT("lowerarm_l")},
+		{TEXT("left_hand_joint"), 				TEXT("hand_l")},
+		{TEXT("left_handIndex_1_joint"),		TEXT("index_01_l")},
+		{TEXT("left_handIndex_2_joint"),		TEXT("index_02_l")},
+		{TEXT("left_handIndex_3_joint"),		TEXT("index_03_l")},
+		{TEXT("left_handMid_1_joint"), 			TEXT("middle_01_l")},
+		{TEXT("left_handMid_2_joint"), 			TEXT("middle_02_l")},
+		{TEXT("left_handMid_3_joint"), 			TEXT("middle_03_l")},
+		{TEXT("left_handPinky_1_joint"),		TEXT("pinky_01_l")},
+		{TEXT("left_handPinky_2_joint"),		TEXT("pinky_02_l")},
+		{TEXT("left_handPinky_3_joint"),		TEXT("pinky_03_l")},
+		{TEXT("left_handRing_1_joint"), 		TEXT("ring_01_l")},
+		{TEXT("left_handRing_2_joint"), 		TEXT("ring_02_l")},
+		{TEXT("left_handRing_3_joint"), 		TEXT("ring_03_l")},
+		{TEXT("left_handThumbStart_joint"),		TEXT("thumb_01_l")},
+		{TEXT("left_handThumb_1_joint"),		TEXT("thumb_02_l")},
+		{TEXT("left_handThumb_2_joint"), 		TEXT("thumb_03_l")},
+		{TEXT("right_shoulder_1_joint"), 		TEXT("clavicle_r")},
+		{TEXT("right_arm_joint"), 				TEXT("upperarm_r")},
+		{TEXT("right_forearm_joint"), 			TEXT("lowerarm_r")},
+		{TEXT("right_hand_joint"), 				TEXT("hand_r")},
+		{TEXT("right_handIndex_1_joint"), 		TEXT("index_01_r")},
+		{TEXT("right_handIndex_2_joint"), 		TEXT("index_02_r")},
+		{TEXT("right_handIndex_3_joint"), 		TEXT("index_03_r")},
+		{TEXT("right_handMid_1_joint"), 		TEXT("middle_01_r")},
+		{TEXT("right_handMid_2_joint"), 		TEXT("middle_02_r")},
+		{TEXT("right_handMid_3_joint"), 		TEXT("middle_03_r")},
+		{TEXT("right_handPinky_1_joint"), 		TEXT("pinky_01_r")},
+		{TEXT("right_handPinky_2_joint"), 		TEXT("pinky_02_r")},
+		{TEXT("right_handPinky_3_joint"), 		TEXT("pinky_03_r")},
+		{TEXT("right_handRing_1_joint"), 		TEXT("ring_01_r")},
+		{TEXT("right_handRing_2_joint"), 		TEXT("ring_02_r")},
+		{TEXT("right_handRing_3_joint"), 		TEXT("ring_03_r")},
+		{TEXT("right_handThumbStart_joint"),	TEXT("thumb_01_r")},
+		{TEXT("right_handThumb_1_joint"), 		TEXT("thumb_02_r")},
+		{TEXT("right_handThumb_2_joint"), 		TEXT("thumb_03_r")},
+		{TEXT("neck_1_joint"), 					TEXT("neck_01")},
+		{TEXT("head_joint"), 					TEXT("head")},
+		{TEXT("left_upLeg_joint"), 				TEXT("thigh_l")},
+		{TEXT("left_leg_joint"), 				TEXT("calf_l")},
+		{TEXT("left_foot_joint"), 				TEXT("foot_l")},
+		{TEXT("left_toes_joint"), 				TEXT("ball_l")},
+		{TEXT("right_upLeg_joint"), 			TEXT("thigh_r")},
+		{TEXT("right_leg_joint"), 				TEXT("calf_r")},
+		{TEXT("right_foot_joint"), 				TEXT("foot_r")},
+		{TEXT("right_toes_joint"), 				TEXT("ball_r")}
+	};
+}
+
+static FTransform GetAppleBoneToComponentTransform(const TArray<int>& ParentIndices, const TArray<FTransform>& Transforms, int BoneIndex)
+{
+	FTransform BoneToComponent = FTransform::Identity;
+	
+	int Index = BoneIndex;
+	while (Index >= 0)
+	{
+		BoneToComponent *= Transforms[Index];
+		Index = ParentIndices[Index];
+	}
+	
+	return BoneToComponent;
+}
+
+static FTransform GetAppleBoneToComponentTransformFromFrameData(const FLiveLinkSkeletonStaticData* InSkeletonData, const FLiveLinkAnimationFrameData* InFrameData, int BoneIndex)
+{
+	return GetAppleBoneToComponentTransform(InSkeletonData->GetBoneParents(), InFrameData->Transforms, BoneIndex);
+}
+
+static FTransform GetAppleBoneToComponentTransformFromRefSkeleton(const FLiveLinkSkeletonStaticData* InSkeletonData, int BoneIndex)
+{
+	IModularFeatures& ModularFeatures = IModularFeatures::Get();
+	if (ModularFeatures.IsModularFeatureAvailable(IAppleARKitPoseTrackingLiveLink::GetModularFeatureName()))
+	{
+		IAppleARKitPoseTrackingLiveLink* BodyTrackLiveLink = &IModularFeatures::Get().GetModularFeature<IAppleARKitPoseTrackingLiveLink>(IAppleARKitPoseTrackingLiveLink::GetModularFeatureName());
+		if (auto RefTransforms = BodyTrackLiveLink->GetRefPoseTransforms())
+		{
+			return GetAppleBoneToComponentTransform(InSkeletonData->GetBoneParents(), *RefTransforms, BoneIndex);
+		}
+	}
+
+	return FTransform::Identity;
+}
+
+static FTransform GetRefMeshBoneToComponent(const FCompactPose& OutPose, int BoneIndex)
+{
+	FTransform Transform = FTransform::Identity;
+	while (BoneIndex >= 0)
+	{
+		FCompactPoseBoneIndex CPIndex = OutPose.GetBoneContainer().MakeCompactPoseIndex(FMeshPoseBoneIndex(BoneIndex));
+		if (CPIndex != INDEX_NONE)
+		{
+			Transform *= OutPose.GetRefPose(CPIndex);
+			BoneIndex = OutPose.GetBoneContainer().GetParentBoneIndex(BoneIndex);
+		}
+		else
+		{
+			break;
+		}
+	}
+
+	return Transform;
+}
+
+void UAppleARKitPoseTrackingLiveLinkRemapAsset::BuildPoseFromAnimationData(float DeltaTime, const FLiveLinkSkeletonStaticData* InSkeletonData, const FLiveLinkAnimationFrameData* InFrameData, FCompactPose& OutPose)
+{
+	check(InSkeletonData);
+	check(InFrameData);
+
+	// Transform Bone Names
+	const TArray<FName>& AppleARKitBoneNames = InSkeletonData->GetBoneNames();
+
+	TArray<FName, TMemStackAllocator<>> MeshBoneNames;
+	MeshBoneNames.Reserve(AppleARKitBoneNames.Num());
+
+	TArray<FTransform> MeshBoneToComponents;
+	const auto MeshBoneNum = OutPose.GetBoneContainer().GetReferenceSkeleton().GetNum();
+	MeshBoneToComponents.Reserve(MeshBoneNum);
+	for (auto i = 0; i < MeshBoneNum; ++i)
+	{
+		MeshBoneToComponents.Emplace(GetRefMeshBoneToComponent(OutPose, i));
+	}
+
+	for (const FName& AppleARKitBoneName : AppleARKitBoneNames)
+	{
+		const FName MeshBoneName = GetRemappedBoneName(AppleARKitBoneName);
+		MeshBoneNames.Add(MeshBoneName);
+	}
+
+	const FQuat ApplyARKitForwardDir = AppleARKitHumanForward.ToOrientationQuat();
+	const FQuat MeshForwardDir = MeshForward.ToOrientationQuat();
+	const auto AppleSpaceToMeshSpace = FTransform(ApplyARKitForwardDir).Inverse() * FTransform(MeshForwardDir);
+
+	for (auto AppleBoneIndex = 0; AppleBoneIndex < MeshBoneNames.Num(); ++AppleBoneIndex)
+	{
+		// for each bone, find out its "offset" from ref pose to the actual pose in apple arkit component space
+		// and then apply it to the unreal skeleton
+		const auto& MeshBoneName = MeshBoneNames[AppleBoneIndex];
+		const auto MeshBoneIndex = OutPose.GetBoneContainer().GetPoseBoneIndexForBoneName(MeshBoneName);
+		if (MeshBoneIndex != INDEX_NONE)
+		{
+            const auto AppleARKitBoneToComponent_Ref = GetAppleBoneToComponentTransformFromRefSkeleton(InSkeletonData, AppleBoneIndex);
+			const auto AppleARKitBoneToComponent = GetAppleBoneToComponentTransformFromFrameData(InSkeletonData, InFrameData, AppleBoneIndex);
+
+			// find out the offset including translation and rotation in apple arkit component space
+			const auto VectorOffset = AppleARKitBoneToComponent.GetTranslation() - AppleARKitBoneToComponent_Ref.GetTranslation();
+			const auto RotationOffset = AppleARKitBoneToComponent.GetRotation() * AppleARKitBoneToComponent_Ref.GetRotation().Inverse();
+            
+			// tranform the ref bone to unreal skeleton component space
+			const auto MeshBoneToComponent_Ref = GetRefMeshBoneToComponent(OutPose, MeshBoneIndex);
+			// convert it into apple arkit component space
+			const auto MeshBoneToComponent_Ref_InAppleSpace = MeshBoneToComponent_Ref * AppleSpaceToMeshSpace.Inverse();
+			// apply the offset, so as to get the posed bone transform
+			const auto MeshBoneToComponent_Ref_InAppleSpace_WithOffset = FTransform(RotationOffset * MeshBoneToComponent_Ref_InAppleSpace.GetRotation()
+                                                                                    , VectorOffset + MeshBoneToComponent_Ref_InAppleSpace.GetTranslation()
+                                                                                    , MeshBoneToComponent_Ref_InAppleSpace.GetScale3D());
+			// convert back to unreal skeleton component space
+			const auto MeshBoneToComponent = MeshBoneToComponent_Ref_InAppleSpace_WithOffset * AppleSpaceToMeshSpace;
+
+			// save all bone-to-component transforms and extract the bone-to-parent transforms later
+			MeshBoneToComponents[MeshBoneIndex] = MeshBoneToComponent;
+		}
+	}
+
+	for (const auto& MeshBoneName : MeshBoneNames)
+	{
+		const auto MeshBoneIndex = OutPose.GetBoneContainer().GetPoseBoneIndexForBoneName(MeshBoneName);
+		if (MeshBoneIndex != INDEX_NONE)
+		{
+			const auto MeshParentBoneIndex = OutPose.GetBoneContainer().GetParentBoneIndex(MeshBoneIndex);
+			FCompactPoseBoneIndex CPIndex = OutPose.GetBoneContainer().MakeCompactPoseIndex(FMeshPoseBoneIndex(MeshBoneIndex));
+			if (CPIndex != INDEX_NONE)
+			{
+				OutPose[CPIndex] = MeshParentBoneIndex != INDEX_NONE
+									? MeshBoneToComponents[MeshBoneIndex] * MeshBoneToComponents[MeshParentBoneIndex].Inverse()
+									: MeshBoneToComponents[MeshBoneIndex];
+			}
+		}
+	}
+}
+
+FName UAppleARKitPoseTrackingLiveLinkRemapAsset::GetRemappedBoneName(FName BoneName) const
+{
+	if (auto Found = AppleARKitBoneNamesToMeshBoneNames.Find(BoneName))
+	{
+		return *Found;
+	}
+	return NAME_None;
+}
+
diff --git a/Engine/Plugins/Runtime/AR/Apple/AppleARKit/Source/AppleARKitPoseTrackingLiveLink/Private/AppleARKitPoseTrackingLiveLinkRemapAsset.h b/Engine/Plugins/Runtime/AR/Apple/AppleARKit/Source/AppleARKitPoseTrackingLiveLink/Private/AppleARKitPoseTrackingLiveLinkRemapAsset.h
new file mode 100644
--- /dev/null
+++ b/Engine/Plugins/Runtime/AR/Apple/AppleARKit/Source/AppleARKitPoseTrackingLiveLink/Private/AppleARKitPoseTrackingLiveLinkRemapAsset.h
@@ -0,0 +1,33 @@
+// Copyright 1998-2019 Epic Games, Inc. All Rights Reserved.
+
+#pragma once
+
+#include "CoreMinimal.h"
+#include "UObject/Package.h"
+#include "UObject/UObjectGlobals.h"
+#include "UObject/ObjectMacros.h"
+#include "LiveLinkRetargetAsset.h"
+#include "AppleARKitPoseTrackingLiveLinkRemapAsset.generated.h"
+
+UCLASS(Blueprintable)
+class UAppleARKitPoseTrackingLiveLinkRemapAsset :
+	public ULiveLinkRetargetAsset
+{
+	GENERATED_BODY()
+
+public:
+	UAppleARKitPoseTrackingLiveLinkRemapAsset(const FObjectInitializer& ObjectInitializer);
+	virtual void BuildPoseFromAnimationData(float DeltaTime, const FLiveLinkSkeletonStaticData* InSkeletonData, const FLiveLinkAnimationFrameData* InFrameData, FCompactPose& OutPose) override;
+
+	UPROPERTY(EditAnywhere, Category="LiveLink")
+	FVector AppleARKitHumanForward = FVector(-1, 0, 0);
+
+	UPROPERTY(EditAnywhere, Category="LiveLink")
+	FVector MeshForward = FVector(0, 1, 0);
+
+	UPROPERTY(EditAnywhere, Category="LiveLink")
+	TMap<FName, FName> AppleARKitBoneNamesToMeshBoneNames;
+private:
+	FName GetRemappedBoneName(FName BoneName) const;
+
+};
diff --git a/Engine/Plugins/Runtime/AR/Apple/AppleARKit/Source/AppleARKitPoseTrackingLiveLink/Private/AppleARKitPoseTrackingLiveLinkSource.cpp b/Engine/Plugins/Runtime/AR/Apple/AppleARKit/Source/AppleARKitPoseTrackingLiveLink/Private/AppleARKitPoseTrackingLiveLinkSource.cpp
new file mode 100644
--- /dev/null
+++ b/Engine/Plugins/Runtime/AR/Apple/AppleARKit/Source/AppleARKitPoseTrackingLiveLink/Private/AppleARKitPoseTrackingLiveLinkSource.cpp
@@ -0,0 +1,146 @@
+// Copyright 1998-2019 Epic Games, Inc. All Rights Reserved.
+
+#include "AppleARKitPoseTrackingLiveLinkSource.h"
+
+#include "Roles/LiveLinkAnimationTypes.h"
+#include "Roles/LiveLinkAnimationRole.h"
+
+TSharedPtr<ILiveLinkSourceARKitPoseTracking> FAppleARKitPoseTrackingLiveLinkSourceFactory::CreateLiveLinkSource()
+{
+	IModularFeatures& ModularFeatures = IModularFeatures::Get();
+
+	if (ModularFeatures.IsModularFeatureAvailable(ILiveLinkClient::ModularFeatureName))
+	{
+		ILiveLinkClient* LiveLinkClient = &IModularFeatures::Get().GetModularFeature<ILiveLinkClient>(ILiveLinkClient::ModularFeatureName);
+		TSharedPtr <ILiveLinkSourceARKitPoseTracking> Source = MakeShareable(new FAppleARKitPoseTrackingLiveLinkSource());
+		LiveLinkClient->AddSource(Source);
+		return Source;
+	}
+	return TSharedPtr<ILiveLinkSourceARKitPoseTracking>();
+}
+
+FAppleARKitPoseTrackingLiveLinkSource::FAppleARKitPoseTrackingLiveLinkSource() :
+	Client(nullptr)
+	, LastFramePublished(0)
+	, bNewLiveLinkClient(false)
+{
+}
+
+void FAppleARKitPoseTrackingLiveLinkSource::ReceiveClient(ILiveLinkClient* InClient, FGuid InSourceGuid)
+{
+	Client = InClient;
+	SourceGuid = InSourceGuid;
+
+	bNewLiveLinkClient = true;
+}
+
+bool FAppleARKitPoseTrackingLiveLinkSource::IsSourceStillValid() const
+{
+	return Client != nullptr;
+}
+
+bool FAppleARKitPoseTrackingLiveLinkSource::RequestSourceShutdown()
+{
+	Client = nullptr;
+	return true;
+}
+
+FText FAppleARKitPoseTrackingLiveLinkSource::GetSourceMachineName() const
+{
+	return FText().FromString(FPlatformProcess::ComputerName());
+}
+
+FText FAppleARKitPoseTrackingLiveLinkSource::GetSourceStatus() const
+{
+	return NSLOCTEXT( "AppleARKitPoseTrackingLiveLink", "AppleARKitPoseTrackingLiveLinkStatus", "Active" );
+}
+
+FText FAppleARKitPoseTrackingLiveLinkSource::GetSourceType() const
+{
+	return NSLOCTEXT( "AppleARKitPoseTrackingLiveLink", "AppleARKitPoseTrackingLiveLinkSourceType", "Apple AR Body Tracking" );
+}
+
+void FAppleARKitPoseTrackingLiveLinkSource::UpdateLiveLink(FName SubjectName, const FTimecode& Timecode, uint32 FrameRate, const FARPose3D& PoseData, const FARPose3D& RefPoseData, FName DeviceId)
+{
+	const FLiveLinkSubjectKey SubjectKey(SourceGuid, SubjectName);
+
+	if (Client)
+	{
+		// Per ReceiveClient initialization:
+		if (bNewLiveLinkClient)
+		{
+			FLiveLinkStaticDataStruct SkeletalData = SetupLiveLinkStaticData(RefPoseData);
+			Client->RemoveSubject_AnyThread(SubjectKey);
+			Client->PushSubjectStaticData_AnyThread(SubjectKey, ULiveLinkAnimationRole::StaticClass(), MoveTemp(SkeletalData));
+			bNewLiveLinkClient = false;
+		}
+
+		// Every frame updates:
+		{
+			FLiveLinkFrameDataStruct LiveLinkFrameDataStruct(FLiveLinkAnimationFrameData::StaticStruct());
+			FLiveLinkAnimationFrameData& LiveLinkAnimationFrameData = *LiveLinkFrameDataStruct.Cast<FLiveLinkAnimationFrameData>();
+			LiveLinkAnimationFrameData.WorldTime = FPlatformTime::Seconds();
+
+			//Copy transforms over to transient structure
+			// Live link transforms need to be in the hierarchical skeleton, so each in the space of its parent.
+			// The hand tracking transforms are in world space.
+			TArray<FTransform>& Transforms = LiveLinkAnimationFrameData.Transforms;
+			Transforms.Reset();
+			for (int32 Index = 0; Index < PoseData.SkeletonDefinition.NumJoints; ++Index)
+			{
+				int32 ParentIndex = PoseData.SkeletonDefinition.ParentIndices[Index];
+				const bool bIsTracked = PoseData.IsJointTracked[Index] && (ParentIndex < 0 || PoseData.IsJointTracked[ParentIndex]);
+				const TArray<FTransform>& JointTransforms = bIsTracked ? PoseData.JointTransforms : RefPoseData.JointTransforms;
+				const FTransform& BoneTransform = JointTransforms[Index];
+
+				if (ParentIndex < 0)
+				{
+					// We are at the root, so use it.
+					Transforms.Emplace(BoneTransform);
+				}
+				else
+				{
+					const FTransform& ParentTransform = JointTransforms[ParentIndex];
+					Transforms.Emplace(BoneTransform * ParentTransform.Inverse());
+				}
+			}
+
+			// Share the data locally with the LiveLink client
+			Client->PushSubjectFrameData_AnyThread(SubjectKey, MoveTemp(LiveLinkFrameDataStruct));
+		}
+	}
+}
+
+FLiveLinkStaticDataStruct FAppleARKitPoseTrackingLiveLinkSource::SetupLiveLinkStaticData(const FARPose3D& RefPoseData)
+{
+	check(IsInGameThread());
+
+	const FARSkeletonDefinition& ARSkeletonDefinition = RefPoseData.SkeletonDefinition;
+
+	FLiveLinkStaticDataStruct LiveLinkSkeletonStaticData;
+	LiveLinkSkeletonStaticData.InitializeWith(FLiveLinkSkeletonStaticData::StaticStruct(), nullptr);
+	FLiveLinkSkeletonStaticData* SkeletonDataPtr = LiveLinkSkeletonStaticData.Cast<FLiveLinkSkeletonStaticData>();
+	check(SkeletonDataPtr);
+
+	SkeletonDataPtr->SetBoneNames(ARSkeletonDefinition.JointNames);
+	SkeletonDataPtr->SetBoneParents(ARSkeletonDefinition.ParentIndices);
+
+	RefTransforms.Reserve(RefPoseData.SkeletonDefinition.NumJoints);
+	for (int32 Index = 0; Index < RefPoseData.SkeletonDefinition.NumJoints; ++Index)
+	{
+		const FTransform& BoneTransform = RefPoseData.JointTransforms[Index];
+		int32 ParentIndex = RefPoseData.SkeletonDefinition.ParentIndices[Index];
+		if (ParentIndex < 0)
+		{
+			// We are at the root, so use it.
+			RefTransforms.Emplace(BoneTransform);
+		}
+		else
+		{
+			const FTransform& ParentTransform = RefPoseData.JointTransforms[ParentIndex];
+			RefTransforms.Emplace(BoneTransform * ParentTransform.Inverse());
+		}
+	}
+
+	return LiveLinkSkeletonStaticData;
+}
\ No newline at end of file
diff --git a/Engine/Plugins/Runtime/AR/Apple/AppleARKit/Source/AppleARKitPoseTrackingLiveLink/Private/AppleARKitPoseTrackingLiveLinkSource.h b/Engine/Plugins/Runtime/AR/Apple/AppleARKit/Source/AppleARKitPoseTrackingLiveLink/Private/AppleARKitPoseTrackingLiveLinkSource.h
new file mode 100644
--- /dev/null
+++ b/Engine/Plugins/Runtime/AR/Apple/AppleARKit/Source/AppleARKitPoseTrackingLiveLink/Private/AppleARKitPoseTrackingLiveLinkSource.h
@@ -0,0 +1,46 @@
+// Copyright 1998-2019 Epic Games, Inc. All Rights Reserved.
+
+#pragma once
+
+#include "CoreMinimal.h"
+#include "AppleARKitPoseTrackingLiveLinkSourceFactory.h"
+#include "ILiveLinkClient.h"
+#include "LiveLinkTypes.h"
+
+
+/** Publishes face blend shapes to LiveLink for use locally */
+class FAppleARKitPoseTrackingLiveLinkSource :
+	public ILiveLinkSourceARKitPoseTracking
+{
+public:
+	FAppleARKitPoseTrackingLiveLinkSource();
+	virtual ~FAppleARKitPoseTrackingLiveLinkSource() {}
+
+private:
+	// ILiveLinkSource interface
+	virtual void ReceiveClient(ILiveLinkClient* InClient, FGuid InSourceGuid) override;
+	virtual bool IsSourceStillValid() const override;
+	virtual bool RequestSourceShutdown() override;
+	virtual FText GetSourceMachineName() const override;
+	virtual FText GetSourceStatus() const override;
+	virtual FText GetSourceType() const override;
+	// End ILiveLinkSource
+
+	void UpdateLiveLink(FName SubjectName, const FTimecode& Timecode, uint32 FrameRate, const FARPose3D& PoseData, const FARPose3D& RefPoseData, FName DeviceId) override;
+	const TArray<FTransform>* GetRefPoseTransforms() const override { return &RefTransforms; }
+	FLiveLinkStaticDataStruct SetupLiveLinkStaticData(const FARPose3D& RefPoseData);
+private:
+
+	/** The local client to push data updates to */
+	ILiveLinkClient* Client;
+
+	/** Our identifier in LiveLink */
+	FGuid SourceGuid;
+
+	/** The last time we sent the data. Used to not send redundant data */
+	uint32 LastFramePublished;
+
+	bool bNewLiveLinkClient;
+
+	TArray<FTransform> RefTransforms;
+};
diff --git a/Engine/Plugins/Runtime/AR/Apple/AppleARKit/Source/AppleARKitPoseTrackingLiveLink/Public/AppleARKitPoseTrackingLiveLinkSourceFactory.h b/Engine/Plugins/Runtime/AR/Apple/AppleARKit/Source/AppleARKitPoseTrackingLiveLink/Public/AppleARKitPoseTrackingLiveLinkSourceFactory.h
new file mode 100644
--- /dev/null
+++ b/Engine/Plugins/Runtime/AR/Apple/AppleARKit/Source/AppleARKitPoseTrackingLiveLink/Public/AppleARKitPoseTrackingLiveLinkSourceFactory.h
@@ -0,0 +1,26 @@
+// Copyright 1998-2019 Epic Games, Inc. All Rights Reserved.
+
+#pragma once
+
+#include "ARSystem.h"
+#include "ARTrackable.h"
+#include "ILiveLinkSource.h"
+
+class UTimecodeProvider;
+
+/** Interface that publishes face ar blend shape information via LiveLink */
+class APPLEARKITPOSETRACKINGLIVELINK_API ILiveLinkSourceARKitPoseTracking :
+	public ILiveLinkSource
+{
+public:
+	virtual void UpdateLiveLink(FName SubjectName, const FTimecode& Timecode, uint32 FrameRate, const FARPose3D& PoseData, const FARPose3D& RefPoseData, FName DeviceID = NAME_None) = 0;
+	virtual const TArray<FTransform>* GetRefPoseTransforms() const = 0;
+};
+
+/** Factory that creates and registers the sources with the LiveLink client */
+class APPLEARKITPOSETRACKINGLIVELINK_API FAppleARKitPoseTrackingLiveLinkSourceFactory
+{
+public:
+	/** Creates a face mesh source that will autobind to the tracked face mesh */
+	static TSharedPtr<ILiveLinkSourceARKitPoseTracking> CreateLiveLinkSource();
+};
diff --git a/Engine/Source/Runtime/Renderer/Private/PostProcess/SceneRenderTargets.cpp b/Engine/Source/Runtime/Renderer/Private/PostProcess/SceneRenderTargets.cpp
--- a/Engine/Source/Runtime/Renderer/Private/PostProcess/SceneRenderTargets.cpp
+++ b/Engine/Source/Runtime/Renderer/Private/PostProcess/SceneRenderTargets.cpp
@@ -128,6 +128,13 @@ static FAutoConsoleVariableRef CVarVirtualTextureFeedbackFactor(
 	ECVF_RenderThreadSafe | ECVF_ReadOnly /*Read-only as shaders are compiled with this value*/
 );
 
+static TAutoConsoleVariable<int32> CVarMobileClearSceneColorWithMaxAlpha(
+	TEXT("r.MobileClearSceneColorWithMaxAlpha"),
+	0,
+	TEXT("Whether the scene color alpha channel on mobile should be cleared with MAX_FLT.\n")
+    TEXT(" 0: disabled (default)\n"),
+	ECVF_RenderThreadSafe);
+
 /** The global render targets used for scene rendering. */
 static TGlobalResource<FSceneRenderTargets> SceneRenderTargetsSingleton;
 
@@ -562,7 +569,17 @@ void FSceneRenderTargets::Allocate(FRHICommandListImmediate& RHICmdList, const F
 	int GBufferFormat = CVarGBufferFormat.GetValueOnRenderThread();
 
 	// Set default clear values
-	SetDefaultColorClear(FClearValueBinding::Black);
+	if (CurrentShadingPath == EShadingPath::Mobile &&
+		CVarMobileClearSceneColorWithMaxAlpha.GetValueOnRenderThread())
+	{
+		// On mobile the scene depth is calculated from the alpha component of the scene color
+		// Use BlackMaxAlpha to ensure un-rendered pixels have max depth...
+		SetDefaultColorClear(FClearValueBinding::BlackMaxAlpha);
+	}
+	else
+	{
+		SetDefaultColorClear(FClearValueBinding::Black);
+	}
 	SetDefaultDepthClear(FClearValueBinding::DepthFar);
 
 	int SceneColorFormat;
