Fix GIsRequestingExit depreciation warnings

diff --git a/Engine/Plugins/Runtime/AR/Apple/AppleARKit/Source/AppleARKitPoseTrackingLiveLink/Private/AppleARKitPoseTrackingLiveLinkImpl.cpp b/Engine/Plugins/Runtime/AR/Apple/AppleARKit/Source/AppleARKitPoseTrackingLiveLink/Private/AppleARKitPoseTrackingLiveLinkImpl.cpp
--- a/Engine/Plugins/Runtime/AR/Apple/AppleARKit/Source/AppleARKitPoseTrackingLiveLink/Private/AppleARKitPoseTrackingLiveLinkImpl.cpp
+++ b/Engine/Plugins/Runtime/AR/Apple/AppleARKit/Source/AppleARKitPoseTrackingLiveLink/Private/AppleARKitPoseTrackingLiveLinkImpl.cpp
@@ -14,7 +14,7 @@ FAppleARKitPoseTrackingLiveLink::FAppleARKitPoseTrackingLiveLink() :
 FAppleARKitPoseTrackingLiveLink::~FAppleARKitPoseTrackingLiveLink()
 {
 	// Should only be called durirng shutdown
-	check(GIsRequestingExit);
+	check(IsEngineExitRequested());
 }
 
 void FAppleARKitPoseTrackingLiveLink::Init()
