Fixing missing includes

diff --git a/Engine/Plugins/Runtime/AR/Apple/AppleARKit/Source/AppleARKitPoseTrackingLiveLink/Private/AppleARKitPoseTrackingLiveLinkRemapAsset.cpp b/Engine/Plugins/Runtime/AR/Apple/AppleARKit/Source/AppleARKitPoseTrackingLiveLink/Private/AppleARKitPoseTrackingLiveLinkRemapAsset.cpp
--- a/Engine/Plugins/Runtime/AR/Apple/AppleARKit/Source/AppleARKitPoseTrackingLiveLink/Private/AppleARKitPoseTrackingLiveLinkRemapAsset.cpp
+++ b/Engine/Plugins/Runtime/AR/Apple/AppleARKit/Source/AppleARKitPoseTrackingLiveLink/Private/AppleARKitPoseTrackingLiveLinkRemapAsset.cpp
@@ -5,6 +5,8 @@
 #include "Features/IModularFeatures.h"
 #include "AppleARKitPoseTrackingLiveLink.h"
 
+#include "BoneIndices.h"
+
 UAppleARKitPoseTrackingLiveLinkRemapAsset::UAppleARKitPoseTrackingLiveLinkRemapAsset(const FObjectInitializer& ObjectInitializer)
 : Super(ObjectInitializer)
 {
diff --git a/Engine/Plugins/Runtime/AR/Apple/AppleARKit/Source/AppleARKitPoseTrackingLiveLink/Private/AppleARKitPoseTrackingLiveLinkSource.cpp b/Engine/Plugins/Runtime/AR/Apple/AppleARKit/Source/AppleARKitPoseTrackingLiveLink/Private/AppleARKitPoseTrackingLiveLinkSource.cpp
--- a/Engine/Plugins/Runtime/AR/Apple/AppleARKit/Source/AppleARKitPoseTrackingLiveLink/Private/AppleARKitPoseTrackingLiveLinkSource.cpp
+++ b/Engine/Plugins/Runtime/AR/Apple/AppleARKit/Source/AppleARKitPoseTrackingLiveLink/Private/AppleARKitPoseTrackingLiveLinkSource.cpp
@@ -5,6 +5,8 @@
 #include "Roles/LiveLinkAnimationTypes.h"
 #include "Roles/LiveLinkAnimationRole.h"
 
+#include "Features/IModularFeatures.h"
+
 TSharedPtr<ILiveLinkSourceARKitPoseTracking> FAppleARKitPoseTrackingLiveLinkSourceFactory::CreateLiveLinkSource()
 {
 	IModularFeatures& ModularFeatures = IModularFeatures::Get();
