Integrating bug fixes from the arkit stream

diff --git a/Engine/Plugins/Runtime/AR/Apple/AppleARKit/Source/AppleARKit/Private/AppleARKitConversion.cpp b/Engine/Plugins/Runtime/AR/Apple/AppleARKit/Source/AppleARKit/Private/AppleARKitConversion.cpp
--- a/Engine/Plugins/Runtime/AR/Apple/AppleARKit/Source/AppleARKit/Private/AppleARKitConversion.cpp
+++ b/Engine/Plugins/Runtime/AR/Apple/AppleARKit/Source/AppleARKit/Private/AppleARKitConversion.cpp
@@ -405,17 +405,25 @@ ARConfiguration* FAppleARKitConversion::ToARConfiguration( UARSessionConfig* Ses
 		default:
 			return nullptr;
 	}
+	
 	if (SessionConfiguration != nullptr)
 	{
 		// Copy / convert properties
 		SessionConfiguration.lightEstimationEnabled = SessionConfig->GetLightEstimationMode() != EARLightEstimationMode::None;
 		SessionConfiguration.providesAudioData = NO;
 		SessionConfiguration.worldAlignment = FAppleARKitConversion::ToARWorldAlignment(SessionConfig->GetWorldAlignment());
-		
+	}
+	
+	return SessionConfiguration;
+}
+
+void FAppleARKitConversion::ConfigureSessionTrackingFeatures(UARSessionConfig* SessionConfig, ARConfiguration* SessionConfiguration)
+{
 #if SUPPORTS_ARKIT_3_0
 		// Enable additional frame semantics for ARKit 3.0
 		if (FAppleARKitAvailability::SupportsARKit30())
 		{
+			EARSessionType SessionType = SessionConfig->GetSessionType();
 			const EARSessionTrackingFeature SessionTrackingFeature = SessionConfig->GetEnabledSessionTrackingFeature();
 			if (SessionTrackingFeature != EARSessionTrackingFeature::None)
 			{
@@ -432,9 +440,6 @@ ARConfiguration* FAppleARKitConversion::ToARConfiguration( UARSessionConfig* Ses
 			}
 		}
 #endif
-	}
-	
-	return SessionConfiguration;
 }
 
 bool FAppleARKitConversion::IsSessionTrackingFeatureSupported(EARSessionType SessionType, EARSessionTrackingFeature SessionTrackingFeature)
diff --git a/Engine/Plugins/Runtime/AR/Apple/AppleARKit/Source/AppleARKit/Private/AppleARKitFrame.cpp b/Engine/Plugins/Runtime/AR/Apple/AppleARKit/Source/AppleARKit/Private/AppleARKitFrame.cpp
--- a/Engine/Plugins/Runtime/AR/Apple/AppleARKit/Source/AppleARKit/Private/AppleARKitFrame.cpp
+++ b/Engine/Plugins/Runtime/AR/Apple/AppleARKit/Source/AppleARKit/Private/AppleARKitFrame.cpp
@@ -3,6 +3,7 @@
 // AppleARKit
 #include "AppleARKitFrame.h"
 #include "AppleARKitModule.h"
+#include "AppleARKitConversion.h"
 #include "Misc/ScopeLock.h"
 
 // Default constructor
diff --git a/Engine/Plugins/Runtime/AR/Apple/AppleARKit/Source/AppleARKit/Private/AppleARKitSystem.cpp b/Engine/Plugins/Runtime/AR/Apple/AppleARKit/Source/AppleARKit/Private/AppleARKitSystem.cpp
--- a/Engine/Plugins/Runtime/AR/Apple/AppleARKit/Source/AppleARKit/Private/AppleARKitSystem.cpp
+++ b/Engine/Plugins/Runtime/AR/Apple/AppleARKit/Source/AppleARKit/Private/AppleARKitSystem.cpp
@@ -1341,6 +1341,9 @@ bool FAppleARKitSystem::Run(UARSessionConfig* SessionConfig)
 			UE_LOG(LogAppleARKit, Error, TEXT("The requested session type is not supported by this device"));
 			return false;
 		}
+		
+		// Configure additional tracking features
+		FAppleARKitConversion::ConfigureSessionTrackingFeatures(SessionConfig, Configuration);
 
 		// Create our ARSessionDelegate
 		if (Delegate == nullptr)
diff --git a/Engine/Plugins/Runtime/AR/Apple/AppleARKit/Source/AppleARKit/Private/AppleARKitVideoOverlay.h b/Engine/Plugins/Runtime/AR/Apple/AppleARKit/Source/AppleARKit/Private/AppleARKitVideoOverlay.h
--- a/Engine/Plugins/Runtime/AR/Apple/AppleARKit/Source/AppleARKit/Private/AppleARKitVideoOverlay.h
+++ b/Engine/Plugins/Runtime/AR/Apple/AppleARKit/Source/AppleARKit/Private/AppleARKitVideoOverlay.h
@@ -6,6 +6,11 @@
 #include "Materials/MaterialInterface.h"
 #include "UObject/ConstructorHelpers.h"
 #include "AppleARKitAvailability.h"
+
+#if SUPPORTS_ARKIT_1_0
+	#import <ARKit/ARKit.h>
+#endif
+
 #include "AppleARKitVideoOverlay.generated.h"
 
 class UARTextureCameraImage;
diff --git a/Engine/Plugins/Runtime/AR/Apple/AppleARKit/Source/AppleARKit/Public/AppleARKitConversion.h b/Engine/Plugins/Runtime/AR/Apple/AppleARKit/Source/AppleARKit/Public/AppleARKitConversion.h
--- a/Engine/Plugins/Runtime/AR/Apple/AppleARKit/Source/AppleARKit/Public/AppleARKitConversion.h
+++ b/Engine/Plugins/Runtime/AR/Apple/AppleARKit/Source/AppleARKit/Public/AppleARKitConversion.h
@@ -245,6 +245,8 @@ struct APPLEARKIT_API FAppleARKitConversion
 #endif
 
 	static ARConfiguration* ToARConfiguration( UARSessionConfig* SessionConfig, TMap< FString, UARCandidateImage* >& CandidateImages, TMap< FString, CGImageRef >& ConvertedCandidateImages, TMap< FString, UARCandidateObject* >& CandidateObjects );
+	
+	static void ConfigureSessionTrackingFeatures(UARSessionConfig* SessionConfig, ARConfiguration* SessionConfiguration);
 
 	#pragma clang diagnostic pop
 #endif
diff --git a/Engine/Source/Runtime/AugmentedReality/Private/ARTrackable.cpp b/Engine/Source/Runtime/AugmentedReality/Private/ARTrackable.cpp
--- a/Engine/Source/Runtime/AugmentedReality/Private/ARTrackable.cpp
+++ b/Engine/Source/Runtime/AugmentedReality/Private/ARTrackable.cpp
@@ -276,9 +276,9 @@ FTransform UARFaceGeometry::GetWorldSpaceEyeTransform(EAREye Eye) const
 {
 	if (Eye == EAREye::LeftEye)
 	{
-		return GetLocalToWorldTransform() * LeftEyeTransform;
+		return LeftEyeTransform * GetLocalToWorldTransform();
 	}
-	return GetLocalToWorldTransform() * RightEyeTransform;
+	return RightEyeTransform * GetLocalToWorldTransform();
 }
 
 UAREnvironmentCaptureProbe::UAREnvironmentCaptureProbe()
