Fix for MainPS_ES2 cook crash on mobile.

diff --git a/Engine/Plugins/Runtime/AR/Apple/AppleARKit/Source/AppleARKit/Private/AppleARKitVideoOverlay.cpp b/Engine/Plugins/Runtime/AR/Apple/AppleARKit/Source/AppleARKit/Private/AppleARKitVideoOverlay.cpp
--- a/Engine/Plugins/Runtime/AR/Apple/AppleARKit/Source/AppleARKit/Private/AppleARKitVideoOverlay.cpp
+++ b/Engine/Plugins/Runtime/AR/Apple/AppleARKit/Source/AppleARKit/Private/AppleARKitVideoOverlay.cpp
@@ -163,7 +163,7 @@ using FARKitCameraOverlayPS = TARKitCameraOverlayPS<false>;
 using FARKitCameraOverlayMobilePS = TARKitCameraOverlayPS<true>;
 
 template<> IMPLEMENT_MATERIAL_SHADER(FARKitCameraOverlayPS, "/Engine/Private/PostProcessMaterialShaders.usf", "MainPS_VideoOverlay", SF_Pixel);
-template<> IMPLEMENT_MATERIAL_SHADER(FARKitCameraOverlayMobilePS, "/Engine/Private/PostProcessMaterialShaders.usf", "MainPS_ES2", SF_Pixel);
+template<> IMPLEMENT_MATERIAL_SHADER(FARKitCameraOverlayMobilePS, "/Engine/Private/PostProcessMaterialShaders.usf", "MainPS, SF_Pixel);
 
 void FAppleARKitVideoOverlay::UpdateOcclusionTextures(const FAppleARKitFrame& Frame)
 {
diff --git a/Engine/Shaders/Private/PostProcessMaterialShaders.usf b/Engine/Shaders/Private/PostProcessMaterialShaders.usf
--- a/Engine/Shaders/Private/PostProcessMaterialShaders.usf
+++ b/Engine/Shaders/Private/PostProcessMaterialShaders.usf
@@ -249,7 +249,7 @@ uint bFlipYAxis;
 //
 // Mobile version
 //
-void MainVS_ES2(
+void MainVS(
 	in float4 InPosition : ATTRIBUTE0,
 	in float2 InTexCoord : ATTRIBUTE1,
 	out FPostProcessMaterialVSToPS Output,
@@ -286,7 +286,7 @@ void MainVS_ES2(
 #endif
 }
 
-void MainPS_ES2(
+void MainPS(
 	in FPostProcessMaterialVSToPS Input,
 	in float2 InViewportUV : TEXCOORD8,
 	out half4 OutColor : SV_Target0
