Fixed build break on mobile.

diff --git a/Engine/Plugins/Runtime/AR/Apple/AppleARKit/Source/AppleARKit/Private/AppleARKitVideoOverlay.cpp b/Engine/Plugins/Runtime/AR/Apple/AppleARKit/Source/AppleARKit/Private/AppleARKitVideoOverlay.cpp
--- a/Engine/Plugins/Runtime/AR/Apple/AppleARKit/Source/AppleARKit/Private/AppleARKitVideoOverlay.cpp
+++ b/Engine/Plugins/Runtime/AR/Apple/AppleARKit/Source/AppleARKit/Private/AppleARKitVideoOverlay.cpp
@@ -163,7 +163,7 @@ using FARKitCameraOverlayPS = TARKitCameraOverlayPS<false>;
 using FARKitCameraOverlayMobilePS = TARKitCameraOverlayPS<true>;
 
 template<> IMPLEMENT_MATERIAL_SHADER(FARKitCameraOverlayPS, "/Engine/Private/PostProcessMaterialShaders.usf", "MainPS_VideoOverlay", SF_Pixel);
-template<> IMPLEMENT_MATERIAL_SHADER(FARKitCameraOverlayMobilePS, "/Engine/Private/PostProcessMaterialShaders.usf", "MainPS, SF_Pixel);
+template<> IMPLEMENT_MATERIAL_SHADER(FARKitCameraOverlayMobilePS, "/Engine/Private/PostProcessMaterialShaders.usf", "MainPS", SF_Pixel);
 
 void FAppleARKitVideoOverlay::UpdateOcclusionTextures(const FAppleARKitFrame& Frame)
 {
