Fixed a member copy step in AppleARKitFrame::operator= that Wanghao reported

diff --git a/Engine/Plugins/Runtime/AppleARKit/Source/AppleARKit/Private/AppleARKitFrame.cpp b/Engine/Plugins/Runtime/AppleARKit/Source/AppleARKit/Private/AppleARKitFrame.cpp
--- a/Engine/Plugins/Runtime/AppleARKit/Source/AppleARKit/Private/AppleARKitFrame.cpp
+++ b/Engine/Plugins/Runtime/AppleARKit/Source/AppleARKit/Private/AppleARKitFrame.cpp
@@ -19,6 +19,7 @@ FAppleARKitFrame::FAppleARKitFrame()
 	, CapturedYImageHeight(0)
 	, CapturedCbCrImageWidth(0)
 	, CapturedCbCrImageHeight(0)
+	, WorldMappingState(EARWorldMappingState::NotAvailable)
 {
 };
 
@@ -43,8 +44,8 @@ EARWorldMappingState ToEARWorldMappingState(ARWorldMappingStatus MapStatus)
 #if SUPPORTS_ARKIT_1_0
 
 FAppleARKitFrame::FAppleARKitFrame( ARFrame* InARFrame, CVMetalTextureCacheRef MetalTextureCache )
-  : Camera( InARFrame.camera )
-  , LightEstimate( InARFrame.lightEstimate )
+	: Camera( InARFrame.camera )
+	, LightEstimate( InARFrame.lightEstimate )
 	, WorldMappingState(EARWorldMappingState::NotAvailable)
 {
 	// Sanity check
@@ -112,17 +113,17 @@ FAppleARKitFrame::FAppleARKitFrame( ARFrame* InARFrame, CVMetalTextureCacheRef M
 }
 
 FAppleARKitFrame::FAppleARKitFrame( const FAppleARKitFrame& Other )
-  : Timestamp( Other.Timestamp )
-  , CapturedYImage( nullptr )
-  , CapturedCbCrImage( nullptr )
-  , CameraImage( nullptr )
-  , CameraDepth( nullptr )
-  , CapturedYImageWidth( Other.CapturedYImageWidth )
-  , CapturedYImageHeight( Other.CapturedYImageHeight )
-  , CapturedCbCrImageWidth( Other.CapturedCbCrImageWidth )
-  , CapturedCbCrImageHeight( Other.CapturedCbCrImageHeight )
-  , Camera( Other.Camera )
-  , LightEstimate( Other.LightEstimate )
+	: Timestamp( Other.Timestamp )
+	, CapturedYImage( nullptr )
+	, CapturedCbCrImage( nullptr )
+	, CameraImage( nullptr )
+	, CameraDepth( nullptr )
+	, CapturedYImageWidth( Other.CapturedYImageWidth )
+	, CapturedYImageHeight( Other.CapturedYImageHeight )
+	, CapturedCbCrImageWidth( Other.CapturedCbCrImageWidth )
+	, CapturedCbCrImageHeight( Other.CapturedCbCrImageHeight )
+	, Camera( Other.Camera )
+	, LightEstimate( Other.LightEstimate )
 	, WorldMappingState(Other.WorldMappingState)
 {
 	if(Other.NativeFrame != nullptr)
@@ -194,6 +195,7 @@ FAppleARKitFrame& FAppleARKitFrame::operator=( const FAppleARKitFrame& Other )
 	CapturedCbCrImageHeight = Other.CapturedCbCrImageHeight;
 	Camera = Other.Camera;
 	LightEstimate = Other.LightEstimate;
+	WorldMappingState = Other.WorldMappingState;
 
 	NativeFrame = Other.NativeFrame;
 
