Added a UPL file to AppleARKit that modifies the plist to include the requirement of ARKit and request for camera access

diff --git a/Engine/Plugins/Runtime/AppleARKit/Source/AppleARKit/AppleARKit.Build.cs b/Engine/Plugins/Runtime/AppleARKit/Source/AppleARKit/AppleARKit.Build.cs
--- a/Engine/Plugins/Runtime/AppleARKit/Source/AppleARKit/AppleARKit.Build.cs
+++ b/Engine/Plugins/Runtime/AppleARKit/Source/AppleARKit/AppleARKit.Build.cs
@@ -1,6 +1,7 @@
 // Copyright 1998-2018 Epic Games, Inc. All Rights Reserved.
 
 using UnrealBuildTool;
+using System.IO;
 
 public class AppleARKit : ModuleRules
 {
@@ -54,6 +55,9 @@ public class AppleARKit : ModuleRules
     		PrivateDependencyModuleNames.Add("IOSRuntimeSettings");
 
 			PublicFrameworks.Add( "ARKit" );
+
+            string PluginPath = Utils.MakePathRelativeTo(ModuleDirectory, Target.RelativeEnginePath);
+            AdditionalPropertiesForReceipt.Add("IOSPlugin", Path.Combine(PluginPath, "AppleARKit_IOS_UPL.xml"));
 		}
 	}
 }
diff --git a/Engine/Plugins/Runtime/AppleARKit/Source/AppleARKit/AppleARKit_IOS_UPL.xml b/Engine/Plugins/Runtime/AppleARKit/Source/AppleARKit/AppleARKit_IOS_UPL.xml
new file mode 100644
--- /dev/null
+++ b/Engine/Plugins/Runtime/AppleARKit/Source/AppleARKit/AppleARKit_IOS_UPL.xml
@@ -0,0 +1,44 @@
+<?xml version="1.0" encoding="utf-8"?>
+<root>
+
+    <init>
+        <log text="AppleARKit adding arkit requirement and camera permission to plist..."/>
+    </init>
+
+    <iosPListUpdates>
+        <addElements tag="dict" once="true">
+            <key>NSCameraUsageDescription</key>
+            <string>The camera is used for augmenting reality.</string>
+        </addElements>
+
+        <loopElements tag="dict">
+            <setBool result="bReqKey" value="false"/>
+
+            <loopElements tag="$">
+                <setStringFromTag result="TagName" tag="$"/>
+
+                <setBoolIsEqual result="bIsKey" arg1="$S(TagName)" arg2="key"/>
+                <if condition="bIsKey">
+                    <true>
+                        <setStringFromTagText result="TagValue" tag="$"/>
+                        <setBoolIsEqual result="bReqKey" arg1="$S(TagValue)" arg2="UIRequiredDeviceCapabilities"/>
+                    </true>
+                </if>
+
+                <setBoolIsEqual result="bIsArray" arg1="$S(TagName)" arg2="array"/>
+                <if condition="bIsArray">
+                    <true>
+                        <if condition="bReqKey">
+                            <true>
+                                <setElement result="NewEntry" xml="&lt;string&gt;arkit&lt;/string&gt;"/>
+                                <addElement tag="$" name="NewEntry"/>
+                            </true>
+                        </if>
+                    </true>
+                </if>
+
+            </loopElements>
+        </loopElements>
+    </iosPListUpdates>
+
+</root>
\ No newline at end of file
diff --git a/Engine/Source/Programs/UnrealBuildTool/Platform/IOS/UEDeployIOS.cs b/Engine/Source/Programs/UnrealBuildTool/Platform/IOS/UEDeployIOS.cs
--- a/Engine/Source/Programs/UnrealBuildTool/Platform/IOS/UEDeployIOS.cs
+++ b/Engine/Source/Programs/UnrealBuildTool/Platform/IOS/UEDeployIOS.cs
@@ -274,12 +274,6 @@ namespace UnrealBuildTool
 				}
 			}
 			ConfigHierarchy GameIni = ConfigCache.ReadHierarchy(ConfigHierarchyType.Game, DirRef, UnrealTargetPlatform.IOS);
-			bool bSupportAR = false;
-			GameIni.GetBool("/Script/EngineSettings.GeneralProjectSettings", "bSupportAR", out bSupportAR);
-			if (bSupportAR)
-			{
-				RequiredCaps += "\t\t<string>arkit</string>\n";
-			}
 
 			Ini.GetBool("/Script/IOSRuntimeSettings.IOSRuntimeSettings", "bSupportsOpenGLES2", out bSupported);
 			RequiredCaps += bSupported ? "\t\t<string>opengles-2</string>\n" : "";
@@ -647,13 +641,6 @@ namespace UnrealBuildTool
 				}
 			}
 
-			// add the camera usage key
-			if (bSupportAR)
-			{
-				Text.AppendLine("\t<key>NSCameraUsageDescription</key>");
-				Text.AppendLine("\t\t<string>The camera is used for augmenting reality.</string>");
-			}
-
 			// Add remote-notifications as background mode
 			if (bRemoteNotificationsSupported)
 			{
