Deprecated bIsTracked on tracked images and face geo per Google API reunification feedback

diff --git a/Engine/Plugins/Runtime/AppleARKit/Source/AppleARKit/Private/AppleARKitSystem.cpp b/Engine/Plugins/Runtime/AppleARKit/Source/AppleARKit/Private/AppleARKitSystem.cpp
--- a/Engine/Plugins/Runtime/AppleARKit/Source/AppleARKit/Private/AppleARKitSystem.cpp
+++ b/Engine/Plugins/Runtime/AppleARKit/Source/AppleARKit/Private/AppleARKitSystem.cpp
@@ -1482,7 +1482,11 @@ void FAppleARKitSystem::SessionDidAddAnchors_Internal( TSharedRef<FAppleARKitAnc
 			NewAnchorDebugName = FString::Printf(TEXT("FACE-%02d"), LastTrackedGeometry_DebugId++);
 			UARFaceGeometry* NewGeo = NewObject<UARFaceGeometry>();
 			NewGeo->UpdateFaceGeometry(SharedThis(this), GameThreadFrameNumber, UpdateTimestamp, AnchorData->Transform, GetAlignmentTransform(), AnchorData->BlendShapes, AnchorData->FaceVerts, AnchorData->FaceIndices, AnchorData->LeftEyeTransform, AnchorData->RightEyeTransform, AnchorData->LookAtTarget);
+			NewGeo->SetTrackingState(EARTrackingState::Tracking);
+			// @todo JoeG -- remove in 4.22
+PRAGMA_DISABLE_DEPRECATION_WARNINGS
 			NewGeo->bIsTracked = true;
+PRAGMA_ENABLE_DEPRECATION_WARNINGS
 			NewGeometry = NewGeo;
 			break;
 		}
@@ -1583,7 +1587,11 @@ void FAppleARKitSystem::SessionDidUpdateAnchors_Internal( TSharedRef<FAppleARKit
 				if (UARFaceGeometry* FaceGeo = Cast<UARFaceGeometry>(FoundGeometry))
 				{
 					FaceGeo->UpdateFaceGeometry(SharedThis(this), GameThreadFrameNumber, UpdateTimestamp, AnchorData->Transform, GetAlignmentTransform(), AnchorData->BlendShapes, AnchorData->FaceVerts, AnchorData->FaceIndices, AnchorData->LeftEyeTransform, AnchorData->RightEyeTransform, AnchorData->LookAtTarget);
+					FaceGeo->SetTrackingState(AnchorData->bIsTracked ? EARTrackingState::Tracking : EARTrackingState::NotTracking);
+					// @todo JoeG -- remove this in 4.22
+PRAGMA_DISABLE_DEPRECATION_WARNINGS
 					FaceGeo->bIsTracked = AnchorData->bIsTracked;
+PRAGMA_ENABLE_DEPRECATION_WARNINGS
 					for (UARPin* Pin : PinsToUpdate)
 					{
 						const FTransform Pin_LocalToTrackingTransform_PostUpdate = Pin->GetLocalToTrackingTransform_NoAlignment() * AnchorDeltaTransform;
@@ -1600,7 +1608,11 @@ void FAppleARKitSystem::SessionDidUpdateAnchors_Internal( TSharedRef<FAppleARKit
 					ensure(CandidateImage != nullptr);
 
 					ImageAnchor->UpdateTrackedGeometry(SharedThis(this), GameThreadFrameNumber, UpdateTimestamp, AnchorData->Transform, GetAlignmentTransform(), *CandidateImage);
+					ImageAnchor->SetTrackingState(AnchorData->bIsTracked ? EARTrackingState::Tracking : EARTrackingState::NotTracking);
+					// @todo JoeG -- remove this in 4.22
+PRAGMA_DISABLE_DEPRECATION_WARNINGS
 					ImageAnchor->bIsTracked = AnchorData->bIsTracked;
+PRAGMA_ENABLE_DEPRECATION_WARNINGS
 					for (UARPin* Pin : PinsToUpdate)
 					{
 						const FTransform Pin_LocalToTrackingTransform_PostUpdate = Pin->GetLocalToTrackingTransform_NoAlignment() * AnchorDeltaTransform;
diff --git a/Engine/Source/Runtime/AugmentedReality/Private/ARTrackable.cpp b/Engine/Source/Runtime/AugmentedReality/Private/ARTrackable.cpp
--- a/Engine/Source/Runtime/AugmentedReality/Private/ARTrackable.cpp
+++ b/Engine/Source/Runtime/AugmentedReality/Private/ARTrackable.cpp
@@ -50,6 +50,16 @@ EARTrackingState UARTrackedGeometry::GetTrackingState() const
 	return TrackingState;
 }
 
+bool UARTrackedGeometry::IsTracked() const
+{
+	return TrackingState == EARTrackingState::Tracking;
+}
+
+void UARTrackedGeometry::SetTrackingState(EARTrackingState NewState)
+{
+	TrackingState = NewState;
+}
+
 FTransform UARTrackedGeometry::GetLocalToWorldTransform() const
 {
 	return GetLocalToTrackingTransform() * GetARSystem()->GetTrackingToWorldTransform();
diff --git a/Engine/Source/Runtime/AugmentedReality/Public/ARTrackable.h b/Engine/Source/Runtime/AugmentedReality/Public/ARTrackable.h
--- a/Engine/Source/Runtime/AugmentedReality/Public/ARTrackable.h
+++ b/Engine/Source/Runtime/AugmentedReality/Public/ARTrackable.h
@@ -40,7 +40,11 @@ public:
 	
 	UFUNCTION(BlueprintPure, Category="AR AugmentedReality|Tracked Geometry")
 	EARTrackingState GetTrackingState() const;
-	
+	void SetTrackingState(EARTrackingState NewState);
+
+	UFUNCTION(BlueprintPure, Category="AR AugmentedReality|Tracked Geometry")
+	bool IsTracked() const;
+
 	UFUNCTION(BlueprintPure, Category="AR AugmentedReality|Tracked Geometry")
 	FName GetDebugName() const;
 	
@@ -143,6 +147,7 @@ public:
 	UFUNCTION(BlueprintPure, Category = "AR AugmentedReality|Image Detection")
 	UARCandidateImage* GetDetectedImage() const { return DetectedImage; };
 
+	DEPRECATED(4.21, "This property is now deprecated, please use GetTrackingState() and check for EARTrackingState::Tracking or IsTracked() instead.")
 	/** Whether the image is currently being tracked by the AR system */
 	UPROPERTY(BlueprintReadOnly, Category="AR AugmentedReality|Face Geometry")
 	bool bIsTracked;
@@ -280,6 +285,7 @@ public:
 	UPROPERTY(BlueprintReadOnly, Category="AR AugmentedReality|Face Geometry")
 	FVector LookAtTarget;
 
+	DEPRECATED(4.21, "This property is now deprecated, please use GetTrackingState() and check for EARTrackingState::Tracking or IsTracked() instead.")
 	/** Whether the face is currently being tracked by the AR system */
 	UPROPERTY(BlueprintReadOnly, Category="AR AugmentedReality|Face Geometry")
 	bool bIsTracked;
diff --git a/Engine/Source/Runtime/AugmentedReality/Public/ARTypes.h b/Engine/Source/Runtime/AugmentedReality/Public/ARTypes.h
--- a/Engine/Source/Runtime/AugmentedReality/Public/ARTypes.h
+++ b/Engine/Source/Runtime/AugmentedReality/Public/ARTypes.h
@@ -18,15 +18,17 @@ class UTexture2D;
 UENUM(BlueprintType, Category="AR AugmentedReality", meta=(Experimental))
 enum class EARTrackingState : uint8
 {
-	/** Currently not tracking. */
+	/** Unknown tracking state */
+	Unknown,
+	
+	/** Currently tracking. */
 	Tracking,
 	
 	/** Currently not tracking, but may resume tracking later. */
 	NotTracking,
 	
 	/** Stopped tracking forever. */
-	StoppedTracking,
-	
+	StoppedTracking
 };
 
 
