Merged Josh's fix for orientation detection

diff --git a/Engine/Plugins/Runtime/AppleARKit/Source/AppleARKit/Private/AppleARKitSystem.cpp b/Engine/Plugins/Runtime/AppleARKit/Source/AppleARKit/Private/AppleARKitSystem.cpp
--- a/Engine/Plugins/Runtime/AppleARKit/Source/AppleARKit/Private/AppleARKitSystem.cpp
+++ b/Engine/Plugins/Runtime/AppleARKit/Source/AppleARKit/Private/AppleARKitSystem.cpp
@@ -1119,77 +1119,15 @@ bool FAppleARKitSystem::HitTestAtScreenPosition(const FVector2D ScreenPosition,
 	return false;
 }
 
-static TOptional<EScreenOrientation::Type> PickAllowedDeviceOrientation( EScreenOrientation::Type InOrientation )
+void FAppleARKitSystem::SetDeviceOrientationAndDerivedTracking(EScreenOrientation::Type InOrientation)
 {
-#if SUPPORTS_ARKIT_1_0
-	const UIOSRuntimeSettings* IOSSettings = GetDefault<UIOSRuntimeSettings>();
-	
-	const bool bOrientationSupported[] =
-	{
-		true, // Unknown
-		IOSSettings->bSupportsPortraitOrientation != 0, // Portait
-		IOSSettings->bSupportsUpsideDownOrientation != 0, // PortraitUpsideDown
-		IOSSettings->bSupportsLandscapeRightOrientation != 0, // LandscapeLeft; These are flipped vs the enum name?
-		IOSSettings->bSupportsLandscapeLeftOrientation != 0, // LandscapeRight; These are flipped vs the enum name?
-		false, // FaceUp
-		false // FaceDown
-	};
-	
-	if (bOrientationSupported[static_cast<int32>(InOrientation)])
-	{
-		return InOrientation;
-	}
-	else
-	{
-		return TOptional<EScreenOrientation::Type>();
-	}
-#else
-	return TOptional<EScreenOrientation::Type>();
-#endif
-}
-
-void FAppleARKitSystem::SetDeviceOrientation( EScreenOrientation::Type InOrientation )
-{
-	TOptional<EScreenOrientation::Type> NewOrientation = PickAllowedDeviceOrientation(InOrientation);
+	checkf(InOrientation != EScreenOrientation::Unknown, TEXT("It is not allowed that Orientation is Unknown"));
 
-	if (!NewOrientation.IsSet() && DeviceOrientation == EScreenOrientation::Unknown)
-	{
-		// We do not currently have a valid orientation, nor did the device provide one.
-		// So pick ANY ALLOWED default.
-		// This only realy happens if the device is face down on something or
-		// in another "useless" state for AR.
-		
-		if (!NewOrientation.IsSet())
-		{
-			NewOrientation = PickAllowedDeviceOrientation(EScreenOrientation::Portrait);
-		}
-		
-		if (!NewOrientation.IsSet())
-		{
-			NewOrientation = PickAllowedDeviceOrientation(EScreenOrientation::LandscapeLeft);
-		}
-		
-		if (!NewOrientation.IsSet())
-		{
-			NewOrientation = PickAllowedDeviceOrientation(EScreenOrientation::PortraitUpsideDown);
-		}
-		
-		if (!NewOrientation.IsSet())
-		{
-			NewOrientation = PickAllowedDeviceOrientation(EScreenOrientation::LandscapeRight);
-		}
-		
-		check(NewOrientation.IsSet());
-	}
-	
-	if (NewOrientation.IsSet() && DeviceOrientation != NewOrientation.GetValue())
-	{
-		DeviceOrientation = NewOrientation.GetValue();
-		CalcTrackingToWorldRotation();
-	}
+	// even if this didn't change, we need to call CalcTrackingToWorldRotation because the camera mode may have changed (from Gravity to non-Gravity, etc)
+	DeviceOrientation = InOrientation;
+	CalcTrackingToWorldRotation();
 }
 
-
 PRAGMA_DISABLE_OPTIMIZATION
 bool FAppleARKitSystem::Run(UARSessionConfig* SessionConfig)
 {
@@ -1200,11 +1138,8 @@ bool FAppleARKitSystem::Run(UARSessionConfig* SessionConfig)
 		LastReceivedFrame = TSharedPtr<FAppleARKitFrame, ESPMode::ThreadSafe>();
 	}
 	
-	// Make sure this is set at session start, because there are timing issues with using only the delegate approach
-	if (DeviceOrientation == EScreenOrientation::Unknown)
-	{
-		SetDeviceOrientation( static_cast<EScreenOrientation::Type>(FPlatformMisc::GetDeviceOrientation()) );
-	}
+	// Make sure this is set at session start, because we could get a camera mode change, without an orientation change
+	SetDeviceOrientationAndDerivedTracking( static_cast<EScreenOrientation::Type>(FPlatformMisc::GetDeviceOrientation()) );
 
 #if SUPPORTS_ARKIT_1_0
 	if (FAppleARKitAvailability::SupportsARKit10())
@@ -1364,7 +1299,7 @@ bool FAppleARKitSystem::Pause()
 void FAppleARKitSystem::OrientationChanged(const int32 NewOrientationRaw)
 {
 	const EScreenOrientation::Type NewOrientation = static_cast<EScreenOrientation::Type>(NewOrientationRaw);
-	SetDeviceOrientation(NewOrientation);
+	SetDeviceOrientationAndDerivedTracking(NewOrientation);
 }
 						
 void FAppleARKitSystem::SessionDidUpdateFrame_DelegateThread(TSharedPtr< FAppleARKitFrame, ESPMode::ThreadSafe > Frame)
diff --git a/Engine/Plugins/Runtime/AppleARKit/Source/AppleARKit/Private/AppleARKitSystem.h b/Engine/Plugins/Runtime/AppleARKit/Source/AppleARKit/Private/AppleARKitSystem.h
--- a/Engine/Plugins/Runtime/AppleARKit/Source/AppleARKit/Private/AppleARKitSystem.h
+++ b/Engine/Plugins/Runtime/AppleARKit/Source/AppleARKit/Private/AppleARKitSystem.h
@@ -132,7 +132,7 @@ private:
 	
 	bool bIsRunning = false;
 	
-	void SetDeviceOrientation( EScreenOrientation::Type InOrientation );
+	void SetDeviceOrientationAndDerivedTracking( EScreenOrientation::Type InOrientation );
 
 	/** Creates or clears the face ar support object if face ar has been requested */
 	void CheckForFaceARSupport(UARSessionConfig* InSessionConfig);
diff --git a/Engine/Source/Runtime/ApplicationCore/Private/IOS/IOSAppDelegate.cpp b/Engine/Source/Runtime/ApplicationCore/Private/IOS/IOSAppDelegate.cpp
--- a/Engine/Source/Runtime/ApplicationCore/Private/IOS/IOSAppDelegate.cpp
+++ b/Engine/Source/Runtime/ApplicationCore/Private/IOS/IOSAppDelegate.cpp
@@ -992,13 +992,15 @@ static FAutoConsoleVariableRef CVarGEnableThermalsReport(
 - (void) didRotate:(NSNotification *)notification
 {   
 #if !PLATFORM_TVOS
-	UIDeviceOrientation orientation = [[UIDevice currentDevice] orientation];
-
+	// get the interfaec orientation
+	NSNumber* OrientationNumber = [notification.userInfo objectForKey:UIApplicationStatusBarOrientationUserInfoKey];
+	UIInterfaceOrientation Orientation = (UIInterfaceOrientation)[OrientationNumber intValue];
+	
     if (bEngineInit)
     {
-		FFunctionGraphTask::CreateAndDispatchWhenReady([orientation]()
+		FFunctionGraphTask::CreateAndDispatchWhenReady([Orientation]()
 		{
-			FCoreDelegates::ApplicationReceivedScreenOrientationChangedNotificationDelegate.Broadcast((int32)orientation);
+			FCoreDelegates::ApplicationReceivedScreenOrientationChangedNotificationDelegate.Broadcast((int32)Orientation);
 
 			//we also want to fire off the safe frame event
 			FCoreDelegates::OnSafeFrameChangedEvent.Broadcast();
diff --git a/Engine/Source/Runtime/Core/Private/IOS/IOSPlatformMisc.cpp b/Engine/Source/Runtime/Core/Private/IOS/IOSPlatformMisc.cpp
--- a/Engine/Source/Runtime/Core/Private/IOS/IOSPlatformMisc.cpp
+++ b/Engine/Source/Runtime/Core/Private/IOS/IOSPlatformMisc.cpp
@@ -275,18 +275,16 @@ bool FIOSPlatformMisc::IsInLowPowerMode()
 
 
 #if !PLATFORM_TVOS
-EDeviceScreenOrientation ConvertFromUIDeviceOrientation(UIDeviceOrientation Orientation)
+EDeviceScreenOrientation ConvertFromUIInterfaceOrientation(UIInterfaceOrientation Orientation)
 {
 	switch(Orientation)
 	{
 		default:
-		case UIDeviceOrientationUnknown : return EDeviceScreenOrientation::Unknown; break;
-		case UIDeviceOrientationPortrait : return EDeviceScreenOrientation::Portrait; break;
-		case UIDeviceOrientationPortraitUpsideDown : return EDeviceScreenOrientation::PortraitUpsideDown; break;
-		case UIDeviceOrientationLandscapeLeft : return EDeviceScreenOrientation::LandscapeLeft; break;
-		case UIDeviceOrientationLandscapeRight : return EDeviceScreenOrientation::LandscapeRight; break;
-		case UIDeviceOrientationFaceUp : return EDeviceScreenOrientation::FaceUp; break;
-		case UIDeviceOrientationFaceDown : return EDeviceScreenOrientation::FaceDown; break;
+		case UIInterfaceOrientationUnknown : return EDeviceScreenOrientation::Unknown; break;
+		case UIInterfaceOrientationPortrait : return EDeviceScreenOrientation::Portrait; break;
+		case UIInterfaceOrientationPortraitUpsideDown : return EDeviceScreenOrientation::PortraitUpsideDown; break;
+		case UIInterfaceOrientationLandscapeLeft : return EDeviceScreenOrientation::LandscapeLeft; break;
+		case UIInterfaceOrientationLandscapeRight : return EDeviceScreenOrientation::LandscapeRight; break;
 	}
 }
 #endif
@@ -294,7 +292,7 @@ EDeviceScreenOrientation ConvertFromUIDeviceOrientation(UIDeviceOrientation Orie
 EDeviceScreenOrientation FIOSPlatformMisc::GetDeviceOrientation()
 {
 #if !PLATFORM_TVOS
-	return ConvertFromUIDeviceOrientation([[UIDevice currentDevice] orientation]);
+	return ConvertFromUIInterfaceOrientation([[UIApplication sharedApplication] statusBarOrientation]);
 #else
 	return EDeviceScreenOrientation::Unknown;
 #endif
