Conversion of ARCore and ARKit to not use FARSystemBase, allowing for other HMDs to use that functionality through composition rather than inheritance

diff --git a/Engine/Plugins/Experimental/OSVR/Source/OSVR/Private/OSVRHMD.cpp b/Engine/Plugins/Experimental/OSVR/Source/OSVR/Private/OSVRHMD.cpp
--- a/Engine/Plugins/Experimental/OSVR/Source/OSVR/Private/OSVRHMD.cpp
+++ b/Engine/Plugins/Experimental/OSVR/Source/OSVR/Private/OSVRHMD.cpp
@@ -544,7 +544,8 @@ bool FOSVRHMD::GetHMDDistortionEnabled(EShadingPath ShadingPath) const
 }
 
 FOSVRHMD::FOSVRHMD(TSharedPtr<class OSVREntryPoint, ESPMode::ThreadSafe> entryPoint) :
-    mOSVREntryPoint(entryPoint)
+	FHeadMountedDisplayBase(nullptr),
+	mOSVREntryPoint(entryPoint)
 {
     static const FName RendererModuleName("Renderer");
     RendererModule = FModuleManager::GetModulePtr<IRendererModule>(RendererModuleName);
diff --git a/Engine/Plugins/Experimental/RemoteSession/Source/RemoteSession/Channels/RemoteSessionXRTrackingChannel.h b/Engine/Plugins/Experimental/RemoteSession/Source/RemoteSession/Channels/RemoteSessionXRTrackingChannel.h
--- a/Engine/Plugins/Experimental/RemoteSession/Source/RemoteSession/Channels/RemoteSessionXRTrackingChannel.h
+++ b/Engine/Plugins/Experimental/RemoteSession/Source/RemoteSession/Channels/RemoteSessionXRTrackingChannel.h
@@ -12,6 +12,10 @@ class REMOTESESSION_API FXRTrackingProxy :
 	public FXRTrackingSystemBase
 {
 public:
+	FXRTrackingProxy()
+		: FXRTrackingSystemBase(nullptr)
+	{}
+		
 	virtual bool IsTracking(int32 DeviceId) override { return true; }
 	virtual bool DoesSupportPositionalTracking() const override { return true; }
 	virtual bool IsHeadTrackingAllowed() const override { return true; }
diff --git a/Engine/Plugins/Experimental/SimpleHMD/Source/SimpleHMD/Private/SimpleHMD.cpp b/Engine/Plugins/Experimental/SimpleHMD/Source/SimpleHMD/Private/SimpleHMD.cpp
--- a/Engine/Plugins/Experimental/SimpleHMD/Source/SimpleHMD/Private/SimpleHMD.cpp
+++ b/Engine/Plugins/Experimental/SimpleHMD/Source/SimpleHMD/Private/SimpleHMD.cpp
@@ -306,6 +306,7 @@ bool FSimpleHMD::IsActiveThisFrame(class FViewport* InViewport) const
 }
 
 FSimpleHMD::FSimpleHMD(const FAutoRegister& AutoRegister) :
+	FHeadMountedDisplayBase(nullptr),
 	FSceneViewExtensionBase(AutoRegister),
 	CurHmdOrientation(FQuat::Identity),
 	LastHmdOrientation(FQuat::Identity),
diff --git a/Engine/Plugins/Lumin/MagicLeap/MagicLeap.uplugin b/Engine/Plugins/Lumin/MagicLeap/MagicLeap.uplugin
--- a/Engine/Plugins/Lumin/MagicLeap/MagicLeap.uplugin
+++ b/Engine/Plugins/Lumin/MagicLeap/MagicLeap.uplugin
@@ -22,7 +22,6 @@
 		{
 			"Name" : "MagicLeap",
 			"Type" : "Runtime",
-			"LoadingPhase" : "PostConfigInit",
 			"WhitelistPlatforms" :
 			[
 				"Win64",
@@ -33,7 +32,6 @@
 		{
 			"Name" : "MagicLeapAudio",
 			"Type" : "Runtime",
-			"LoadingPhase" : "PreDefault",
 			"WhitelistPlatforms" :
 			[
 				"Win64",
@@ -43,7 +41,6 @@
 		{
 			"Name" : "MagicLeapController",
 			"Type" : "Runtime",
-			"LoadingPhase" : "PostConfigInit",
 			"WhitelistPlatforms" :
 			[
 				"Win64",
@@ -54,7 +51,6 @@
 		{
 			"Name" : "MagicLeapEyeTracker",
 			"Type" : "Runtime",
-			"LoadingPhase" : "PostConfigInit",
 			"WhitelistPlatforms" :
 			[
 				"Win64",
@@ -65,7 +61,6 @@
 		{
 			"Name" : "MagicLeapHandTracking",
 			"Type" : "Runtime",
-			"LoadingPhase" : "PostConfigInit",
 			"WhitelistPlatforms" :
 			[
 				"Win64",
@@ -76,7 +71,6 @@
 		{
 			"Name" : "MagicLeapMusicService",
 			"Type" : "Runtime",
-			"LoadingPhase" : "PostConfigInit",
 			"WhitelistPlatforms" :
 			[
 				"Win64",
@@ -87,7 +81,6 @@
 		{
 			"Name" : "MagicLeapSecureStorage",
 			"Type" : "Runtime",
-			"LoadingPhase" : "PostConfigInit",
 			"WhitelistPlatforms" :
 			[
 				"Win64",
@@ -98,7 +91,6 @@
 		{
 			"Name" : "MagicLeapHelperOpenGL",
 			"Type" : "Runtime",
-			"LoadingPhase" : "PostConfigInit",
 			"WhitelistPlatforms" :
 			[
 				"Win64",
@@ -109,7 +101,6 @@
 		{
 			"Name" : "MagicLeapHelperVulkan",
 			"Type" : "Runtime",
-			"LoadingPhase" : "PostConfigInit",
 			"WhitelistPlatforms" :
 			[
 				"Win64",
@@ -120,7 +111,6 @@
 		{
 			"Name" : "MagicLeapIdentity",
 			"Type" : "Runtime",
-			"LoadingPhase" : "PostConfigInit",
 			"WhitelistPlatforms" :
 			[
 				"Win64",
diff --git a/Engine/Plugins/Lumin/MagicLeap/Source/MagicLeap/MagicLeap.Build.cs b/Engine/Plugins/Lumin/MagicLeap/Source/MagicLeap/MagicLeap.Build.cs
--- a/Engine/Plugins/Lumin/MagicLeap/Source/MagicLeap/MagicLeap.Build.cs
+++ b/Engine/Plugins/Lumin/MagicLeap/Source/MagicLeap/MagicLeap.Build.cs
@@ -26,7 +26,8 @@ namespace UnrealBuildTool.Rules
 					"HeadMountedDisplay",
 					"ProceduralMeshComponent",
 					"InputDevice",
-                    "LuminRuntimeSettings"
+                    "LuminRuntimeSettings",
+					"AugmentedReality",
                 }
 			);
 
diff --git a/Engine/Plugins/Lumin/MagicLeap/Source/MagicLeap/Private/MagicLeapHMD.cpp b/Engine/Plugins/Lumin/MagicLeap/Source/MagicLeap/Private/MagicLeapHMD.cpp
--- a/Engine/Plugins/Lumin/MagicLeap/Source/MagicLeap/Private/MagicLeapHMD.cpp
+++ b/Engine/Plugins/Lumin/MagicLeap/Source/MagicLeap/Private/MagicLeapHMD.cpp
@@ -242,7 +242,10 @@ private:
 #if !PLATFORM_MAC
 		if (!HMD.IsValid())
 		{
-			HMD = MakeShared<FMagicLeapHMD, ESPMode::ThreadSafe>(this, bIsVDZIEnabled, bUseVulkanForZI);
+			//initialize AR system
+			TSharedPtr<IARSystemSupport, ESPMode::ThreadSafe> ARImplementation;
+
+			HMD = MakeShared<FMagicLeapHMD, ESPMode::ThreadSafe>(this, ARImplementation.Get(), bIsVDZIEnabled, bUseVulkanForZI);
 		}
 #endif
 #if WITH_EDITOR
@@ -277,6 +280,14 @@ FName FMagicLeapHMD::GetSystemName() const
 	return SystemName;
 }
 
+FString FMagicLeapHMD::GetVersionString() const
+{
+	FString s = FString::Printf(TEXT("LuminHMD - %s, built %s, %s"), *FEngineVersion::Current().ToString(),
+		UTF8_TO_TCHAR(__DATE__), UTF8_TO_TCHAR(__TIME__));
+
+	return s;
+}
+
 bool FMagicLeapHMD::OnStartGameFrame(FWorldContext& WorldContext)
 {
 #if WITH_MLSDK
@@ -319,6 +330,9 @@ bool FMagicLeapHMD::OnStartGameFrame(FWorldContext& WorldContext)
 
 	RefreshTrackingToWorldTransform(WorldContext);
 
+	//update AR system
+	GetARCompositionComponent()->StartARGameFrame(WorldContext);
+
 #endif //WITH_MLSDK
 
 	return true;
@@ -960,8 +974,9 @@ bool FMagicLeapHMD::AllocateRenderTargetTexture(uint32 Index, uint32 SizeX, uint
 //	return GetActiveCustomPresent();
 //}
 
-FMagicLeapHMD::FMagicLeapHMD(IMagicLeapPlugin* InMagicLeapPlugin, bool bEnableVDZI, bool bUseVulkan) :
+FMagicLeapHMD::FMagicLeapHMD(IMagicLeapPlugin* InMagicLeapPlugin, IARSystemSupport* InARImplementation, bool bEnableVDZI, bool bUseVulkan) :
 	// We don't do any mirroring on Lumin as we render direct to the device only.
+	FHeadMountedDisplayBase(InARImplementation),
 #if PLATFORM_LUMIN
 	WindowMirrorMode(0),
 #else
diff --git a/Engine/Plugins/Lumin/MagicLeap/Source/MagicLeap/Private/MagicLeapHMD.h b/Engine/Plugins/Lumin/MagicLeap/Source/MagicLeap/Private/MagicLeapHMD.h
--- a/Engine/Plugins/Lumin/MagicLeap/Source/MagicLeap/Private/MagicLeapHMD.h
+++ b/Engine/Plugins/Lumin/MagicLeap/Source/MagicLeap/Private/MagicLeapHMD.h
@@ -43,6 +43,7 @@ public:
 	virtual class TSharedPtr< class IStereoRendering, ESPMode::ThreadSafe > GetStereoRenderingDevice() override { return AsShared(); }
 	virtual class TSharedPtr< class IXRCamera, ESPMode::ThreadSafe > GetXRCamera(int32 DeviceId) override;
 	virtual FName GetSystemName() const override;
+	virtual FString GetVersionString() const override;
 
 	virtual bool DoesSupportPositionalTracking() const override;
 	virtual bool HasValidTrackingPosition() override;
@@ -136,7 +137,7 @@ public:
 
 public:
 	/** Constructor */
-	FMagicLeapHMD(IMagicLeapPlugin* MagicLeapPlugin, bool bEnableVDZI = false, bool bUseVulkan = false);
+	FMagicLeapHMD(IMagicLeapPlugin* MagicLeapPlugin, IARSystemSupport* ARImplementation, bool bEnableVDZI = false, bool bUseVulkan = false);
 
 	/** Destructor */
 	virtual ~FMagicLeapHMD();
diff --git a/Engine/Plugins/Lumin/MagicLeap/Source/MagicLeapController/Private/TouchpadGesturesComponent.cpp b/Engine/Plugins/Lumin/MagicLeap/Source/MagicLeapController/Private/TouchpadGesturesComponent.cpp
--- a/Engine/Plugins/Lumin/MagicLeap/Source/MagicLeapController/Private/TouchpadGesturesComponent.cpp
+++ b/Engine/Plugins/Lumin/MagicLeap/Source/MagicLeapController/Private/TouchpadGesturesComponent.cpp
@@ -12,19 +12,25 @@ UTouchpadGesturesComponent::UTouchpadGesturesComponent()
 	PrimaryComponentTick.TickGroup = TG_PrePhysics;
 	bAutoActivate = true;
 
-	TSharedPtr<FMagicLeapController> controller = StaticCastSharedPtr<FMagicLeapController>(IMagicLeapControllerPlugin::Get().GetInputDevice());
-	if (controller.IsValid())
+	if (!HasAnyFlags(RF_ClassDefaultObject))
 	{
-		controller->RegisterTouchpadGestureReceiver(this);
-	}	
+		TSharedPtr<FMagicLeapController> controller = StaticCastSharedPtr<FMagicLeapController>(IMagicLeapControllerPlugin::Get().GetInputDevice());
+		if (controller.IsValid())
+		{
+			controller->RegisterTouchpadGestureReceiver(this);
+		}
+	}
 }
 
 UTouchpadGesturesComponent::~UTouchpadGesturesComponent()
 {
-	TSharedPtr<FMagicLeapController> controller = StaticCastSharedPtr<FMagicLeapController>(IMagicLeapControllerPlugin::Get().GetInputDevice());
-	if (controller.IsValid())
+	if (!HasAnyFlags(RF_ClassDefaultObject))
 	{
-		controller->UnregisterTouchpadGestureReceiver(this);
+		TSharedPtr<FMagicLeapController> controller = StaticCastSharedPtr<FMagicLeapController>(IMagicLeapControllerPlugin::Get().GetInputDevice());
+		if (controller.IsValid())
+		{
+			controller->UnregisterTouchpadGestureReceiver(this);
+		}
 	}
 }
 
diff --git a/Engine/Plugins/Runtime/AppleARKit/Source/AppleARKit/Private/AppleARKitModule.cpp b/Engine/Plugins/Runtime/AppleARKit/Source/AppleARKit/Private/AppleARKitModule.cpp
--- a/Engine/Plugins/Runtime/AppleARKit/Source/AppleARKit/Private/AppleARKitModule.cpp
+++ b/Engine/Plugins/Runtime/AppleARKit/Source/AppleARKit/Private/AppleARKitModule.cpp
@@ -15,6 +15,7 @@ TSharedPtr<class IXRTrackingSystem, ESPMode::ThreadSafe> FAppleARKitModule::Crea
 {
 #if PLATFORM_IOS
 	auto NewARKitSystem = AppleARKitSupport::CreateAppleARKitSystem();
+	NewARKitSystem->GetARCompositionComponent()->InitializeARSystem();
 	FAppleARKitARKitSystemPtr = NewARKitSystem;
     return NewARKitSystem;
 #else
diff --git a/Engine/Plugins/Runtime/AppleARKit/Source/AppleARKit/Private/AppleARKitSystem.cpp b/Engine/Plugins/Runtime/AppleARKit/Source/AppleARKit/Private/AppleARKitSystem.cpp
--- a/Engine/Plugins/Runtime/AppleARKit/Source/AppleARKit/Private/AppleARKitSystem.cpp
+++ b/Engine/Plugins/Runtime/AppleARKit/Source/AppleARKit/Private/AppleARKitSystem.cpp
@@ -90,7 +90,7 @@ private:
 	{
 		// @todo arkit : is it safe not to lock here? Theoretically this should only be called on the game thread.
 		ensure(IsInGameThread());
-		const bool bShouldOverrideFOV = ARKitSystem.GetSessionConfig().ShouldRenderCameraOverlay();
+		const bool bShouldOverrideFOV = ARKitSystem.GetARCompositionComponent()->GetSessionConfig().ShouldRenderCameraOverlay();
 		if (bShouldOverrideFOV && ARKitSystem.GameThreadFrame.IsValid())
 		{
 			if (ARKitSystem.DeviceOrientation == EScreenOrientation::Portrait || ARKitSystem.DeviceOrientation == EScreenOrientation::PortraitUpsideDown)
@@ -165,7 +165,7 @@ private:
 		// Most Face AR apps that are driving other meshes using the face capture (animoji style) will disable this.
 		const bool bRenderOverlay =
 			ARKitSystem.OnGetARSessionStatus().Status == EARSessionStatus::Running &&
-			ARKitSystem.GetSessionConfig().ShouldRenderCameraOverlay();
+			ARKitSystem.GetARCompositionComponent()->GetSessionConfig().ShouldRenderCameraOverlay();
 
 #if SUPPORTS_ARKIT_1_0
 		if (FAppleARKitAvailability::SupportsARKit10())
@@ -196,7 +196,7 @@ private:
 //
 
 FAppleARKitSystem::FAppleARKitSystem()
-: FARSystemBase()
+: FXRTrackingSystemBase(this)
 , DeviceOrientation(EScreenOrientation::Unknown)
 , DerivedTrackingToUnrealRotation(FRotator::ZeroRotator)
 , LightEstimate(nullptr)
@@ -266,7 +266,7 @@ bool FAppleARKitSystem::GetCurrentPose(int32 DeviceId, FQuat& OutOrientation, FV
 		// Apply alignment transform if there is one.
 		FTransform CurrentTransform(GameThreadFrame->Camera.Orientation, GameThreadFrame->Camera.Translation);
 		CurrentTransform = FTransform(DerivedTrackingToUnrealRotation) * CurrentTransform;
-		CurrentTransform *= GetAlignmentTransform();
+		CurrentTransform *= GetARCompositionComponent()->GetAlignmentTransform();
 		
 		
 		// Apply counter-rotation to compensate for mobile device orientation
@@ -303,7 +303,7 @@ void FAppleARKitSystem::CalcTrackingToWorldRotation()
 	// We rotate the camera to counteract the portrait vs. landscape viewport rotation
 	DerivedTrackingToUnrealRotation = FRotator::ZeroRotator;
 
-	const EARWorldAlignment WorldAlignment = GetSessionConfig().GetWorldAlignment();
+	const EARWorldAlignment WorldAlignment = GetARCompositionComponent()->GetSessionConfig().GetWorldAlignment();
 	if (WorldAlignment == EARWorldAlignment::Gravity || WorldAlignment == EARWorldAlignment::GravityAndHeading)
 	{
 		switch (DeviceOrientation)
@@ -407,7 +407,7 @@ bool FAppleARKitSystem::IsHeadTrackingAllowed() const
 	// For face AR tracking movements of the device most likely won't want to be tracked
 	const bool bEnableCameraTracking =
 		OnGetARSessionStatus().Status == EARSessionStatus::Running &&
-		GetSessionConfig().ShouldEnableCameraTracking();
+		GetARCompositionComponent()->GetSessionConfig().ShouldEnableCameraTracking();
 
 #if SUPPORTS_ARKIT_1_0
 	if (FAppleARKitAvailability::SupportsARKit10())
@@ -550,8 +550,6 @@ void FAppleARKitSystem::OnSetAlignmentTransform(const FTransform& InAlignmentTra
 	{
 		Pin->UpdateAlignmentTransform(NewAlignmentTransform);
 	}
-	
-	SetAlignmentTransform_Internal(InAlignmentTransform);
 }
 
 static bool IsHitInRange( float UnrealHitDistance )
@@ -589,7 +587,7 @@ TArray<FARTraceResult> FAppleARKitSystem::OnLineTraceTrackedObjects( const FVect
 	{
 #if SUPPORTS_ARKIT_1_0
 		
-		TSharedRef<FARSystemBase, ESPMode::ThreadSafe> This = SharedThis(this);
+		TSharedPtr<FARSupportInterface , ESPMode::ThreadSafe> This = GetARCompositionComponent();
 		
 		@autoreleasepool
 		{
@@ -634,7 +632,7 @@ TArray<FARTraceResult> FAppleARKitSystem::OnLineTraceTrackedObjects( const FVect
 						{
 							// Hit result has passed and above filtering, add it to the list
 							// Convert to Unreal's Hit Test result format
-							Results.Add( FARTraceResult( This, UnrealHitDistance, EARLineTraceChannels::PlaneUsingExtent, FAppleARKitConversion::ToFTransform( HitTestResult.worldTransform )*GetAlignmentTransform(), FindGeometryFromAnchor(HitTestResult.anchor, TrackedGeometries) ) );
+							Results.Add(FARTraceResult(This, UnrealHitDistance, EARLineTraceChannels::PlaneUsingExtent, FAppleARKitConversion::ToFTransform(HitTestResult.worldTransform)*GetARCompositionComponent()->GetAlignmentTransform(), FindGeometryFromAnchor(HitTestResult.anchor, TrackedGeometries)));
 						}
 					}
 				}
@@ -650,7 +648,7 @@ TArray<FARTraceResult> FAppleARKitSystem::OnLineTraceTrackedObjects( const FVect
 						{
 							// Hit result has passed and above filtering, add it to the list
 							// Convert to Unreal's Hit Test result format
-							Results.Add( FARTraceResult( This, UnrealHitDistance, EARLineTraceChannels::GroundPlane, FAppleARKitConversion::ToFTransform( HitTestResult.worldTransform )*GetAlignmentTransform(), FindGeometryFromAnchor(HitTestResult.anchor, TrackedGeometries) ) );
+							Results.Add(FARTraceResult(This, UnrealHitDistance, EARLineTraceChannels::GroundPlane, FAppleARKitConversion::ToFTransform(HitTestResult.worldTransform)*GetARCompositionComponent()->GetAlignmentTransform(), FindGeometryFromAnchor(HitTestResult.anchor, TrackedGeometries)));
 						}
 					}
 				}
@@ -668,7 +666,7 @@ TArray<FARTraceResult> FAppleARKitSystem::OnLineTraceTrackedObjects( const FVect
 						{
 							// Hit result has passed and above filtering, add it to the list
 							// Convert to Unreal's Hit Test result format
-							Results.Add( FARTraceResult( This, UnrealHitDistance, EARLineTraceChannels::FeaturePoint, FAppleARKitConversion::ToFTransform( HitTestResult.worldTransform )*GetAlignmentTransform(), FindGeometryFromAnchor(HitTestResult.anchor, TrackedGeometries) ) );
+							Results.Add( FARTraceResult( This, UnrealHitDistance, EARLineTraceChannels::FeaturePoint, FAppleARKitConversion::ToFTransform( HitTestResult.worldTransform )*GetARCompositionComponent()->GetAlignmentTransform(), FindGeometryFromAnchor(HitTestResult.anchor, TrackedGeometries) ) );
 						}
 					}
 				}
@@ -730,7 +728,7 @@ UARPin* FAppleARKitSystem::OnPinComponent( USceneComponent* ComponentToPin, cons
 		// PinToWorld * AlignedTrackingToWorld(-1) * TrackingToAlignedTracking(-1) = PinToWorld * WorldToAlignedTracking * AlignedTrackingToTracking
 		// The Worlds and AlignedTracking cancel out, and we get PinToTracking
 		// But we must translate this logic into Unreal's transform API
-		const FTransform& TrackingToAlignedTracking = GetAlignmentTransform();
+		const FTransform& TrackingToAlignedTracking = GetARCompositionComponent()->GetAlignmentTransform();
 		const FTransform PinToTrackingTransform = PinToWorldTransform.GetRelativeTransform(GetTrackingToWorldTransform()).GetRelativeTransform(TrackingToAlignedTracking);
 
 		// If the user did not provide a TrackedGeometry, create the simplest TrackedGeometry for this pin.
@@ -740,11 +738,11 @@ UARPin* FAppleARKitSystem::OnPinComponent( USceneComponent* ComponentToPin, cons
 			double UpdateTimestamp = FPlatformTime::Seconds();
 			
 			GeometryToPinTo = NewObject<UARTrackedPoint>();
-			GeometryToPinTo->UpdateTrackedGeometry(SharedThis(this), GameThreadFrameNumber, UpdateTimestamp, PinToTrackingTransform, GetAlignmentTransform());
+			GeometryToPinTo->UpdateTrackedGeometry(GetARCompositionComponent().ToSharedRef(), GameThreadFrameNumber, UpdateTimestamp, PinToTrackingTransform, GetARCompositionComponent()->GetAlignmentTransform());
 		}
 		
 		UARPin* NewPin = NewObject<UARPin>();
-		NewPin->InitARPin(SharedThis(this), ComponentToPin, PinToTrackingTransform, GeometryToPinTo, DebugName);
+		NewPin->InitARPin(GetARCompositionComponent().ToSharedRef(), ComponentToPin, PinToTrackingTransform, GeometryToPinTo, DebugName);
 		
 		Pins.Add(NewPin);
 		
@@ -1107,8 +1105,6 @@ EARWorldMappingState FAppleARKitSystem::OnGetWorldMappingStatus() const
 
 void FAppleARKitSystem::AddReferencedObjects( FReferenceCollector& Collector )
 {
-	FARSystemBase::AddReferencedObjects(Collector);
-
 	Collector.AddReferencedObjects( TrackedGeometries );
 	Collector.AddReferencedObjects( Pins );
 	Collector.AddReferencedObject( CameraImage );
@@ -1310,7 +1306,7 @@ bool FAppleARKitSystem::Run(UARSessionConfig* SessionConfig)
 	// Set running state
 	bIsRunning = true;
 	
-	OnARSessionStarted.Broadcast();
+	GetARCompositionComponent()->OnARSessionStarted.Broadcast();
 	return true;
 }
 PRAGMA_ENABLE_OPTIMIZATION
@@ -1477,8 +1473,8 @@ void FAppleARKitSystem::SessionDidAddAnchors_DelegateThread( NSArray<ARAnchor*>*
 	if (FaceARSupport != nullptr)
 	{
 		double UpdateTimestamp = FPlatformTime::Seconds();
-		const FRotator& AdjustBy = GetSessionConfig().GetWorldAlignment() == EARWorldAlignment::Camera ? DerivedTrackingToUnrealRotation : FRotator::ZeroRotator;
-		const EARFaceTrackingUpdate UpdateSetting = GetSessionConfig().GetFaceTrackingUpdate();
+		const FRotator& AdjustBy = GetARCompositionComponent()->GetSessionConfig().GetWorldAlignment() == EARWorldAlignment::Camera ? DerivedTrackingToUnrealRotation : FRotator::ZeroRotator;
+		const EARFaceTrackingUpdate UpdateSetting = GetARCompositionComponent()->GetSessionConfig().GetFaceTrackingUpdate();
 
 		const TArray<TSharedPtr<FAppleARKitAnchorData>> AnchorList = FaceARSupport->MakeAnchorData(anchors, UpdateTimestamp, GameThreadFrameNumber, AdjustBy, UpdateSetting);
 		for (TSharedPtr<FAppleARKitAnchorData> NewAnchorData : AnchorList)
@@ -1506,8 +1502,8 @@ void FAppleARKitSystem::SessionDidUpdateAnchors_DelegateThread( NSArray<ARAnchor
 	if (FaceARSupport != nullptr)
 	{
 		double UpdateTimestamp = FPlatformTime::Seconds();
-		const FRotator& AdjustBy = GetSessionConfig().GetWorldAlignment() == EARWorldAlignment::Camera ? DerivedTrackingToUnrealRotation : FRotator::ZeroRotator;
-		const EARFaceTrackingUpdate UpdateSetting = GetSessionConfig().GetFaceTrackingUpdate();
+		const FRotator& AdjustBy = GetARCompositionComponent()->GetSessionConfig().GetWorldAlignment() == EARWorldAlignment::Camera ? DerivedTrackingToUnrealRotation : FRotator::ZeroRotator;
+		const EARFaceTrackingUpdate UpdateSetting = GetARCompositionComponent()->GetSessionConfig().GetFaceTrackingUpdate();
 
 		const TArray<TSharedPtr<FAppleARKitAnchorData>> AnchorList = FaceARSupport->MakeAnchorData(anchors, UpdateTimestamp, GameThreadFrameNumber, AdjustBy, UpdateSetting);
 		for (TSharedPtr<FAppleARKitAnchorData> NewAnchorData : AnchorList)
@@ -1547,8 +1543,10 @@ void FAppleARKitSystem::SessionDidAddAnchors_Internal( TSharedRef<FAppleARKitAnc
 {
 	double UpdateTimestamp = FPlatformTime::Seconds();
 	
+	const TSharedPtr<FARSupportInterface , ESPMode::ThreadSafe>& ARComponent = GetARCompositionComponent();
+
 	// In case we have camera tracking turned off, we still need to update the frame
-	if (!GetSessionConfig().ShouldEnableCameraTracking())
+	if (!ARComponent->GetSessionConfig().ShouldEnableCameraTracking())
 	{
 		UpdateFrame();
 	}
@@ -1567,14 +1565,14 @@ void FAppleARKitSystem::SessionDidAddAnchors_Internal( TSharedRef<FAppleARKitAnc
 		{
 			NewAnchorDebugName = FString::Printf(TEXT("ANCHOR-%02d"), LastTrackedGeometry_DebugId++);
 			NewGeometry = NewObject<UARTrackedGeometry>();
-			NewGeometry->UpdateTrackedGeometry(SharedThis(this), GameThreadFrameNumber, UpdateTimestamp, AnchorData->Transform, GetAlignmentTransform());
+			NewGeometry->UpdateTrackedGeometry(ARComponent.ToSharedRef(), GameThreadFrameNumber, UpdateTimestamp, AnchorData->Transform, GetARCompositionComponent()->GetAlignmentTransform());
 			break;
 		}
 		case EAppleAnchorType::PlaneAnchor:
 		{
 			NewAnchorDebugName = FString::Printf(TEXT("PLN-%02d"), LastTrackedGeometry_DebugId++);
 			UARPlaneGeometry* NewGeo = NewObject<UARPlaneGeometry>();
-			NewGeo->UpdateTrackedGeometry(SharedThis(this), GameThreadFrameNumber, UpdateTimestamp, AnchorData->Transform, GetAlignmentTransform(), AnchorData->Center, AnchorData->Extent);
+			NewGeo->UpdateTrackedGeometry(ARComponent.ToSharedRef(), GameThreadFrameNumber, UpdateTimestamp, AnchorData->Transform, GetARCompositionComponent()->GetAlignmentTransform(), AnchorData->Center, AnchorData->Extent);
 			NewGeometry = NewGeo;
 			break;
 		}
@@ -1582,7 +1580,7 @@ void FAppleARKitSystem::SessionDidAddAnchors_Internal( TSharedRef<FAppleARKitAnc
 		{
 			NewAnchorDebugName = FString::Printf(TEXT("FACE-%02d"), LastTrackedGeometry_DebugId++);
 			UARFaceGeometry* NewGeo = NewObject<UARFaceGeometry>();
-			NewGeo->UpdateFaceGeometry(SharedThis(this), GameThreadFrameNumber, UpdateTimestamp, AnchorData->Transform, GetAlignmentTransform(), AnchorData->BlendShapes, AnchorData->FaceVerts, AnchorData->FaceIndices, AnchorData->LeftEyeTransform, AnchorData->RightEyeTransform, AnchorData->LookAtTarget);
+			NewGeo->UpdateFaceGeometry(ARComponent.ToSharedRef(), GameThreadFrameNumber, UpdateTimestamp, AnchorData->Transform, GetARCompositionComponent()->GetAlignmentTransform(), AnchorData->BlendShapes, AnchorData->FaceVerts, AnchorData->FaceIndices, AnchorData->LeftEyeTransform, AnchorData->RightEyeTransform, AnchorData->LookAtTarget);
 			NewGeo->SetTrackingState(EARTrackingState::Tracking);
 			// @todo JoeG -- remove in 4.22
 PRAGMA_DISABLE_DEPRECATION_WARNINGS
@@ -1597,7 +1595,7 @@ PRAGMA_ENABLE_DEPRECATION_WARNINGS
 			UARTrackedImage* NewImage = NewObject<UARTrackedImage>();
 			UARCandidateImage** CandidateImage = CandidateImages.Find(AnchorData->DetectedAnchorName);
 			ensure(CandidateImage != nullptr);
-			NewImage->UpdateTrackedGeometry(SharedThis(this), GameThreadFrameNumber, UpdateTimestamp, AnchorData->Transform, GetAlignmentTransform(), *CandidateImage);
+			NewImage->UpdateTrackedGeometry(ARComponent.ToSharedRef(), GameThreadFrameNumber, UpdateTimestamp, AnchorData->Transform, GetARCompositionComponent()->GetAlignmentTransform(), *CandidateImage);
 			NewGeometry = NewImage;
 			break;
 		}
@@ -1605,7 +1603,7 @@ PRAGMA_ENABLE_DEPRECATION_WARNINGS
 		{
 			NewAnchorDebugName = FString::Printf(TEXT("ENV-%02d"), LastTrackedGeometry_DebugId++);
 			UAppleARKitEnvironmentCaptureProbe* NewProbe = NewObject<UAppleARKitEnvironmentCaptureProbe>();
-			NewProbe->UpdateEnvironmentCapture(SharedThis(this), GameThreadFrameNumber, UpdateTimestamp, AnchorData->Transform, GetAlignmentTransform(), AnchorData->Extent, AnchorData->ProbeTexture);
+			NewProbe->UpdateEnvironmentCapture(ARComponent.ToSharedRef(), GameThreadFrameNumber, UpdateTimestamp, AnchorData->Transform, GetARCompositionComponent()->GetAlignmentTransform(), AnchorData->Extent, AnchorData->ProbeTexture);
 			NewGeometry = NewProbe;
 			break;
 		}
@@ -1615,7 +1613,7 @@ PRAGMA_ENABLE_DEPRECATION_WARNINGS
 			UARTrackedObject* NewTrackedObject = NewObject<UARTrackedObject>();
 			UARCandidateObject** CandidateObject = CandidateObjects.Find(AnchorData->DetectedAnchorName);
 			ensure(CandidateObject != nullptr);
-			NewTrackedObject->UpdateTrackedGeometry(SharedThis(this), GameThreadFrameNumber, UpdateTimestamp, AnchorData->Transform, GetAlignmentTransform(), *CandidateObject);
+			NewTrackedObject->UpdateTrackedGeometry(ARComponent.ToSharedRef(), GameThreadFrameNumber, UpdateTimestamp, AnchorData->Transform, GetARCompositionComponent()->GetAlignmentTransform(), *CandidateObject);
 			NewGeometry = NewTrackedObject;
 			break;
 		}
@@ -1631,8 +1629,10 @@ void FAppleARKitSystem::SessionDidUpdateAnchors_Internal( TSharedRef<FAppleARKit
 {
 	double UpdateTimestamp = FPlatformTime::Seconds();
 	
+	const TSharedPtr<FARSupportInterface , ESPMode::ThreadSafe>& ARComponent = GetARCompositionComponent();
+
 	// In case we have camera tracking turned off, we still need to update the frame
-	if (!GetSessionConfig().ShouldEnableCameraTracking())
+	if (!ARComponent->GetSessionConfig().ShouldEnableCameraTracking())
 	{
 		UpdateFrame();
 	}
@@ -1661,7 +1661,7 @@ void FAppleARKitSystem::SessionDidUpdateAnchors_Internal( TSharedRef<FAppleARKit
 		{
 			case EAppleAnchorType::Anchor:
 			{
-				FoundGeometry->UpdateTrackedGeometry(SharedThis(this), GameThreadFrameNumber, UpdateTimestamp, AnchorData->Transform, GetAlignmentTransform());
+				FoundGeometry->UpdateTrackedGeometry(ARComponent.ToSharedRef(), GameThreadFrameNumber, UpdateTimestamp, AnchorData->Transform, GetARCompositionComponent()->GetAlignmentTransform());
 				for (UARPin* Pin : PinsToUpdate)
 				{
 					const FTransform Pin_LocalToTrackingTransform_PostUpdate = Pin->GetLocalToTrackingTransform_NoAlignment() * AnchorDeltaTransform;
@@ -1674,7 +1674,7 @@ void FAppleARKitSystem::SessionDidUpdateAnchors_Internal( TSharedRef<FAppleARKit
 			{
 				if (UARPlaneGeometry* PlaneGeo = Cast<UARPlaneGeometry>(FoundGeometry))
 				{
-					PlaneGeo->UpdateTrackedGeometry(SharedThis(this), GameThreadFrameNumber, UpdateTimestamp, AnchorData->Transform, GetAlignmentTransform(), AnchorData->Center, AnchorData->Extent, AnchorData->BoundaryVerts, nullptr);
+					PlaneGeo->UpdateTrackedGeometry(ARComponent.ToSharedRef(), GameThreadFrameNumber, UpdateTimestamp, AnchorData->Transform, GetARCompositionComponent()->GetAlignmentTransform(), AnchorData->Center, AnchorData->Extent, AnchorData->BoundaryVerts, nullptr);
 					for (UARPin* Pin : PinsToUpdate)
 					{
 						const FTransform Pin_LocalToTrackingTransform_PostUpdate = Pin->GetLocalToTrackingTransform_NoAlignment() * AnchorDeltaTransform;
@@ -1687,7 +1687,7 @@ void FAppleARKitSystem::SessionDidUpdateAnchors_Internal( TSharedRef<FAppleARKit
 			{
 				if (UARFaceGeometry* FaceGeo = Cast<UARFaceGeometry>(FoundGeometry))
 				{
-					FaceGeo->UpdateFaceGeometry(SharedThis(this), GameThreadFrameNumber, UpdateTimestamp, AnchorData->Transform, GetAlignmentTransform(), AnchorData->BlendShapes, AnchorData->FaceVerts, AnchorData->FaceIndices, AnchorData->LeftEyeTransform, AnchorData->RightEyeTransform, AnchorData->LookAtTarget);
+					FaceGeo->UpdateFaceGeometry(ARComponent.ToSharedRef(), GameThreadFrameNumber, UpdateTimestamp, AnchorData->Transform, GetARCompositionComponent()->GetAlignmentTransform(), AnchorData->BlendShapes, AnchorData->FaceVerts, AnchorData->FaceIndices, AnchorData->LeftEyeTransform, AnchorData->RightEyeTransform, AnchorData->LookAtTarget);
 					FaceGeo->SetTrackingState(AnchorData->bIsTracked ? EARTrackingState::Tracking : EARTrackingState::NotTracking);
 					// @todo JoeG -- remove this in 4.22
 PRAGMA_DISABLE_DEPRECATION_WARNINGS
@@ -1708,7 +1708,7 @@ PRAGMA_ENABLE_DEPRECATION_WARNINGS
 					UARCandidateImage** CandidateImage = CandidateImages.Find(AnchorData->DetectedAnchorName);
 					ensure(CandidateImage != nullptr);
 
-					ImageAnchor->UpdateTrackedGeometry(SharedThis(this), GameThreadFrameNumber, UpdateTimestamp, AnchorData->Transform, GetAlignmentTransform(), *CandidateImage);
+					ImageAnchor->UpdateTrackedGeometry(ARComponent.ToSharedRef(), GameThreadFrameNumber, UpdateTimestamp, AnchorData->Transform, GetARCompositionComponent()->GetAlignmentTransform(), *CandidateImage);
 					ImageAnchor->SetTrackingState(AnchorData->bIsTracked ? EARTrackingState::Tracking : EARTrackingState::NotTracking);
 					// @todo JoeG -- remove this in 4.22
 PRAGMA_DISABLE_DEPRECATION_WARNINGS
@@ -1727,7 +1727,7 @@ PRAGMA_ENABLE_DEPRECATION_WARNINGS
 				if (UAppleARKitEnvironmentCaptureProbe* ProbeAnchor = Cast<UAppleARKitEnvironmentCaptureProbe>(FoundGeometry))
 				{
 					// NOTE: The metal texture will be a different texture every time the cubemap is updated which requires a render resource flush
-					ProbeAnchor->UpdateEnvironmentCapture(SharedThis(this), GameThreadFrameNumber, UpdateTimestamp, AnchorData->Transform, GetAlignmentTransform(), AnchorData->Extent, AnchorData->ProbeTexture);
+					ProbeAnchor->UpdateEnvironmentCapture(ARComponent.ToSharedRef(), GameThreadFrameNumber, UpdateTimestamp, AnchorData->Transform, GetARCompositionComponent()->GetAlignmentTransform(), AnchorData->Extent, AnchorData->ProbeTexture);
 					for (UARPin* Pin : PinsToUpdate)
 					{
 						const FTransform Pin_LocalToTrackingTransform_PostUpdate = Pin->GetLocalToTrackingTransform_NoAlignment() * AnchorDeltaTransform;
@@ -1742,8 +1742,10 @@ PRAGMA_ENABLE_DEPRECATION_WARNINGS
 
 void FAppleARKitSystem::SessionDidRemoveAnchors_Internal( FGuid AnchorGuid )
 {
+	const TSharedPtr<FARSupportInterface , ESPMode::ThreadSafe>& ARComponent = GetARCompositionComponent();
+
 	// In case we have camera tracking turned off, we still need to update the frame
-	if (!GetSessionConfig().ShouldEnableCameraTracking())
+	if (!ARComponent->GetSessionConfig().ShouldEnableCameraTracking())
 	{
 		UpdateFrame();
 	}
@@ -1871,7 +1873,7 @@ namespace AppleARKitSupport
 		// Handle older iOS devices somehow calling this
 		if (FAppleARKitAvailability::SupportsARKit10())
 		{
-            auto NewARKitSystem = NewARSystem<FAppleARKitSystem>();
+			auto NewARKitSystem = MakeShared<FAppleARKitSystem, ESPMode::ThreadSafe>();
             return NewARKitSystem;
 		}
 #endif
diff --git a/Engine/Plugins/Runtime/AppleARKit/Source/AppleARKit/Private/AppleARKitSystem.h b/Engine/Plugins/Runtime/AppleARKit/Source/AppleARKit/Private/AppleARKitSystem.h
--- a/Engine/Plugins/Runtime/AppleARKit/Source/AppleARKit/Private/AppleARKitSystem.h
+++ b/Engine/Plugins/Runtime/AppleARKit/Source/AppleARKit/Private/AppleARKitSystem.h
@@ -25,8 +25,7 @@ DECLARE_STATS_GROUP(TEXT("ARKit"), STATGROUP_ARKIT, STATCAT_Advanced);
 struct FAppleARKitFrame;
 struct FAppleARKitAnchorData;
 
-class FAppleARKitSystem :
-	public FARSystemBase
+class FAppleARKitSystem : public IARSystemSupport, public FXRTrackingSystemBase, public FGCObject, public TSharedFromThis<FAppleARKitSystem, ESPMode::ThreadSafe>
 {
 	friend class FAppleARKitXRCamera;
 	
diff --git a/Engine/Plugins/Runtime/AppleARKit/Source/AppleARKit/Private/AppleARKitTrackable.cpp b/Engine/Plugins/Runtime/AppleARKit/Source/AppleARKit/Private/AppleARKitTrackable.cpp
--- a/Engine/Plugins/Runtime/AppleARKit/Source/AppleARKit/Private/AppleARKitTrackable.cpp
+++ b/Engine/Plugins/Runtime/AppleARKit/Source/AppleARKit/Private/AppleARKitTrackable.cpp
@@ -15,7 +15,7 @@ UAppleARKitEnvironmentCaptureProbe::UAppleARKitEnvironmentCaptureProbe()
 }
 
 #if PLATFORM_IOS
-void UAppleARKitEnvironmentCaptureProbe::UpdateEnvironmentCapture(const TSharedRef<FARSystemBase, ESPMode::ThreadSafe>& InTrackingSystem, uint32 FrameNumber, double InTimestamp, const FTransform& InLocalToTrackingTransform, const FTransform& InAlignmentTransform, FVector InExtent, id<MTLTexture> InMetalTexture)
+void UAppleARKitEnvironmentCaptureProbe::UpdateEnvironmentCapture(const TSharedRef<FARSupportInterface , ESPMode::ThreadSafe>& InTrackingSystem, uint32 FrameNumber, double InTimestamp, const FTransform& InLocalToTrackingTransform, const FTransform& InAlignmentTransform, FVector InExtent, id<MTLTexture> InMetalTexture)
 {
 	Super::UpdateEnvironmentCapture(InTrackingSystem, FrameNumber, InTimestamp, InLocalToTrackingTransform, InAlignmentTransform, InExtent);
 	
diff --git a/Engine/Plugins/Runtime/AppleARKit/Source/AppleARKit/Public/AppleARKitTrackable.h b/Engine/Plugins/Runtime/AppleARKit/Source/AppleARKit/Public/AppleARKitTrackable.h
--- a/Engine/Plugins/Runtime/AppleARKit/Source/AppleARKit/Public/AppleARKitTrackable.h
+++ b/Engine/Plugins/Runtime/AppleARKit/Source/AppleARKit/Public/AppleARKitTrackable.h
@@ -19,7 +19,7 @@ public:
 	
 #if PLATFORM_IOS
 	/** Updates the current capture with the new metal texture. This will trigger a render resource update if the textures have changed */
-	void UpdateEnvironmentCapture(const TSharedRef<FARSystemBase, ESPMode::ThreadSafe>& InTrackingSystem, uint32 FrameNumber, double Timestamp, const FTransform& InLocalToTrackingTransform, const FTransform& InAlignmentTransform, FVector InExtent, id<MTLTexture> InMetalTexture);
+	void UpdateEnvironmentCapture(const TSharedRef<FARSupportInterface , ESPMode::ThreadSafe>& InTrackingSystem, uint32 FrameNumber, double Timestamp, const FTransform& InLocalToTrackingTransform, const FTransform& InAlignmentTransform, FVector InExtent, id<MTLTexture> InMetalTexture);
 #endif
 
 private:
diff --git a/Engine/Plugins/Runtime/GoogleARCore/Source/GoogleARCoreBase/Private/GoogleARCoreAPI.cpp b/Engine/Plugins/Runtime/GoogleARCore/Source/GoogleARCoreBase/Private/GoogleARCoreAPI.cpp
--- a/Engine/Plugins/Runtime/GoogleARCore/Source/GoogleARCoreBase/Private/GoogleARCoreAPI.cpp
+++ b/Engine/Plugins/Runtime/GoogleARCore/Source/GoogleARCoreBase/Private/GoogleARCoreAPI.cpp
@@ -1380,8 +1380,8 @@ void UGoogleARCorePointCloud::GetPoint(int Index, FVector& OutWorldPosition, flo
 
 		Point = Point * Session.Pin()->GetWorldToMeterScale();
 		FTransform PointLocalTransform(Point);
-		TSharedRef<FARSystemBase, ESPMode::ThreadSafe> ARSystem = Session.Pin()->GetARSystem();
-		FTransform PointWorldTransform = PointLocalTransform * ARSystem->GetAlignmentTransform() * ARSystem->GetTrackingToWorldTransform();
+		TSharedRef<FARSupportInterface , ESPMode::ThreadSafe> ARSystem = Session.Pin()->GetARSystem();
+		FTransform PointWorldTransform = PointLocalTransform * ARSystem->GetAlignmentTransform() * ARSystem->GetXRTrackingSystem()->GetTrackingToWorldTransform();
 		Point = PointWorldTransform.GetTranslation();
 		Confidence = PointData[Index * 4 + 3];
 #endif
diff --git a/Engine/Plugins/Runtime/GoogleARCore/Source/GoogleARCoreBase/Private/GoogleARCoreAPI.h b/Engine/Plugins/Runtime/GoogleARCore/Source/GoogleARCoreBase/Private/GoogleARCoreAPI.h
--- a/Engine/Plugins/Runtime/GoogleARCore/Source/GoogleARCoreBase/Private/GoogleARCoreAPI.h
+++ b/Engine/Plugins/Runtime/GoogleARCore/Source/GoogleARCoreBase/Private/GoogleARCoreAPI.h
@@ -174,8 +174,8 @@ public:
 	EGoogleARCoreAPIStatus GetSessionCreateStatus();
 	UGoogleARCoreUObjectManager* GetUObjectManager();
 	float GetWorldToMeterScale();
-	void SetARSystem(TSharedRef<FARSystemBase, ESPMode::ThreadSafe> InArSystem) { ARSystem = InArSystem; }
-	TSharedRef<FARSystemBase, ESPMode::ThreadSafe> GetARSystem() { return ARSystem.ToSharedRef(); }
+	void SetARSystem(TSharedRef<FARSupportInterface , ESPMode::ThreadSafe> InArSystem) { ARSystem = InArSystem; }
+	TSharedRef<FARSupportInterface , ESPMode::ThreadSafe> GetARSystem() { return ARSystem.ToSharedRef(); }
 #if PLATFORM_ANDROID
 	ArSession* GetHandle();
 #endif
@@ -214,7 +214,7 @@ private:
 	float CachedWorldToMeterScale;
 	uint32 FrameNumber;
 
-	TSharedPtr<FARSystemBase, ESPMode::ThreadSafe> ARSystem;
+	TSharedPtr<FARSupportInterface , ESPMode::ThreadSafe> ARSystem;
 
 #if PLATFORM_ANDROID
 	ArSession* SessionHandle = nullptr;
diff --git a/Engine/Plugins/Runtime/GoogleARCore/Source/GoogleARCoreBase/Private/GoogleARCoreAugmentedImage.cpp b/Engine/Plugins/Runtime/GoogleARCore/Source/GoogleARCoreBase/Private/GoogleARCoreAugmentedImage.cpp
--- a/Engine/Plugins/Runtime/GoogleARCore/Source/GoogleARCoreBase/Private/GoogleARCoreAugmentedImage.cpp
+++ b/Engine/Plugins/Runtime/GoogleARCore/Source/GoogleARCoreBase/Private/GoogleARCoreAugmentedImage.cpp
@@ -68,7 +68,7 @@ void UGoogleARCoreAugmentedImage::DebugDraw(
 }
 
 void UGoogleARCoreAugmentedImage::UpdateTrackedGeometry(
-	const TSharedRef<FARSystemBase, ESPMode::ThreadSafe>& InTrackingSystem,
+	const TSharedRef<FARSupportInterface , ESPMode::ThreadSafe>& InTrackingSystem,
 	uint32 FrameNumber, double Timestamp, const FTransform& InLocalToTrackingTransform,
 	const FTransform& InAlignmentTransform, const FVector &InCenter, const FVector &InExtent,
 	int32 InImageIndex, const FString& InImageName)
diff --git a/Engine/Plugins/Runtime/GoogleARCore/Source/GoogleARCoreBase/Private/GoogleARCoreBaseModule.cpp b/Engine/Plugins/Runtime/GoogleARCore/Source/GoogleARCoreBase/Private/GoogleARCoreBaseModule.cpp
--- a/Engine/Plugins/Runtime/GoogleARCore/Source/GoogleARCoreBase/Private/GoogleARCoreBaseModule.cpp
+++ b/Engine/Plugins/Runtime/GoogleARCore/Source/GoogleARCoreBase/Private/GoogleARCoreBaseModule.cpp
@@ -52,8 +52,9 @@ IMPLEMENT_MODULE(FGoogleARCoreBaseModule, GoogleARCoreBase)
 TSharedPtr< class IXRTrackingSystem, ESPMode::ThreadSafe > FGoogleARCoreBaseModule::CreateTrackingSystem()
 {
 #if PLATFORM_ANDROID
-	TSharedPtr<FGoogleARCoreXRTrackingSystem, ESPMode::ThreadSafe> ARCoreSystem = NewARSystem<FGoogleARCoreXRTrackingSystem>();
-	FGoogleARCoreDevice::GetInstance()->SetARSystem(ARCoreSystem);
+	auto ARCoreSystem = MakeShared<FGoogleARCoreXRTrackingSystem, ESPMode::ThreadSafe>();
+	ARCoreSystem->GetARCompositionComponent()->InitializeARSystem();
+	FGoogleARCoreDevice::GetInstance()->SetARSystem(ARCoreSystem->GetARCompositionComponent());
 	return ARCoreSystem;
 #else
 	return TSharedPtr<class IXRTrackingSystem, ESPMode::ThreadSafe>();
diff --git a/Engine/Plugins/Runtime/GoogleARCore/Source/GoogleARCoreBase/Private/GoogleARCoreDevice.cpp b/Engine/Plugins/Runtime/GoogleARCore/Source/GoogleARCoreBase/Private/GoogleARCoreDevice.cpp
--- a/Engine/Plugins/Runtime/GoogleARCore/Source/GoogleARCoreBase/Private/GoogleARCoreDevice.cpp
+++ b/Engine/Plugins/Runtime/GoogleARCore/Source/GoogleARCoreBase/Private/GoogleARCoreDevice.cpp
@@ -424,7 +424,7 @@ void FGoogleARCoreDevice::StartSession()
 	ARSystem->OnARSessionStarted.Broadcast();
 }
 
-void FGoogleARCoreDevice::SetARSystem(TSharedPtr<FARSystemBase, ESPMode::ThreadSafe> InARSystem)
+void FGoogleARCoreDevice::SetARSystem(TSharedPtr<FARSupportInterface , ESPMode::ThreadSafe> InARSystem)
 {
 	check(InARSystem.IsValid());
 	ARSystem = InARSystem;
@@ -446,7 +446,7 @@ void* FGoogleARCoreDevice::GetGameThreadARFrameRawPointer()
 	return nullptr;
 }
 
-TSharedPtr<FARSystemBase, ESPMode::ThreadSafe> FGoogleARCoreDevice::GetARSystem()
+TSharedPtr<FARSupportInterface , ESPMode::ThreadSafe> FGoogleARCoreDevice::GetARSystem()
 {
 	return ARSystem;
 }
@@ -623,7 +623,7 @@ EGoogleARCoreFunctionStatus FGoogleARCoreDevice::CreateARPin(const FTransform& P
 	}
 
 	const FTransform& TrackingToAlignedTracking = ARSystem->GetAlignmentTransform();
-	const FTransform PinToTrackingTransform = PinToWorldTransform.GetRelativeTransform(ARSystem->GetTrackingToWorldTransform()).GetRelativeTransform(TrackingToAlignedTracking);
+	const FTransform PinToTrackingTransform = PinToWorldTransform.GetRelativeTransform(ARSystem->GetXRTrackingSystem()->GetTrackingToWorldTransform()).GetRelativeTransform(TrackingToAlignedTracking);
 
 	EGoogleARCoreFunctionStatus Status = ToARCoreFunctionStatus(ARCoreSession->CreateARAnchor(PinToTrackingTransform, TrackedGeometry, ComponentToPin, DebugName, OutARAnchorObject));
 
diff --git a/Engine/Plugins/Runtime/GoogleARCore/Source/GoogleARCoreBase/Private/GoogleARCoreDevice.h b/Engine/Plugins/Runtime/GoogleARCore/Source/GoogleARCoreBase/Private/GoogleARCoreDevice.h
--- a/Engine/Plugins/Runtime/GoogleARCore/Source/GoogleARCoreBase/Private/GoogleARCoreDevice.h
+++ b/Engine/Plugins/Runtime/GoogleARCore/Source/GoogleARCoreBase/Private/GoogleARCoreDevice.h
@@ -114,8 +114,8 @@ public:
 	// Function that is used to call from the Android UI thread:
 	void StartSessionWithRequestedConfig();
 
-	TSharedPtr<FARSystemBase, ESPMode::ThreadSafe> GetARSystem();
-	void SetARSystem(TSharedPtr<FARSystemBase, ESPMode::ThreadSafe> InARSystem);
+	TSharedPtr<FARSupportInterface , ESPMode::ThreadSafe> GetARSystem();
+	void SetARSystem(TSharedPtr<FARSupportInterface , ESPMode::ThreadSafe> InARSystem);
 
 	void* GetARSessionRawPointer();
 	void* GetGameThreadARFrameRawPointer();
@@ -175,5 +175,5 @@ private:
 
 	TQueue<TFunction<void()>> RunOnGameThreadQueue;
 
-	TSharedPtr<FARSystemBase, ESPMode::ThreadSafe> ARSystem;
+	TSharedPtr<FARSupportInterface , ESPMode::ThreadSafe> ARSystem;
 };
diff --git a/Engine/Plugins/Runtime/GoogleARCore/Source/GoogleARCoreBase/Private/GoogleARCoreFunctionLibrary.cpp b/Engine/Plugins/Runtime/GoogleARCore/Source/GoogleARCoreBase/Private/GoogleARCoreFunctionLibrary.cpp
--- a/Engine/Plugins/Runtime/GoogleARCore/Source/GoogleARCoreBase/Private/GoogleARCoreFunctionLibrary.cpp
+++ b/Engine/Plugins/Runtime/GoogleARCore/Source/GoogleARCoreBase/Private/GoogleARCoreFunctionLibrary.cpp
@@ -191,7 +191,7 @@ UGoogleARCoreEventManager* UGoogleARCoreSessionFunctionLibrary::GetARCoreEventMa
 	auto ARSystem = FGoogleARCoreDevice::GetInstance()->GetARSystem();
 	if (ARSystem.IsValid())
 	{
-		return static_cast<FGoogleARCoreXRTrackingSystem*>(ARSystem.Get())->GetEventManager();
+		return static_cast<FGoogleARCoreXRTrackingSystem*>(ARSystem->GetXRTrackingSystem())->GetEventManager();
 	}
 	else
 	{
diff --git a/Engine/Plugins/Runtime/GoogleARCore/Source/GoogleARCoreBase/Private/GoogleARCoreXRTrackingSystem.cpp b/Engine/Plugins/Runtime/GoogleARCore/Source/GoogleARCoreBase/Private/GoogleARCoreXRTrackingSystem.cpp
--- a/Engine/Plugins/Runtime/GoogleARCore/Source/GoogleARCoreBase/Private/GoogleARCoreXRTrackingSystem.cpp
+++ b/Engine/Plugins/Runtime/GoogleARCore/Source/GoogleARCoreBase/Private/GoogleARCoreXRTrackingSystem.cpp
@@ -10,7 +10,8 @@
 #include "ARSessionConfig.h"
 
 FGoogleARCoreXRTrackingSystem::FGoogleARCoreXRTrackingSystem()
-	: ARCoreDeviceInstance(nullptr)
+	: FXRTrackingSystemBase(this)
+	, ARCoreDeviceInstance(nullptr)
 	, bMatchDeviceCameraFOV(false)
 	, bEnablePassthroughCameraRendering(false)
 	, bHasValidPose(false)
@@ -90,7 +91,7 @@ bool FGoogleARCoreXRTrackingSystem::OnStartGameFrame(FWorldContext& WorldContext
 		if (ARCoreDeviceInstance->GetTrackingState() == EGoogleARCoreTrackingState::Tracking)
 		{
 			CurrentPose = ARCoreDeviceInstance->GetLatestPose();
-			CurrentPose *= GetAlignmentTransform();
+			CurrentPose *= GetARCompositionComponent()->GetAlignmentTransform();
 			bHasValidPose = true;
 			CachedTrackingToWorld = ComputeTrackingToWorldTransform(WorldContext);
 		}
@@ -217,19 +218,17 @@ void FGoogleARCoreXRTrackingSystem::OnSetAlignmentTransform(const FTransform& In
 {
 	const FTransform& NewAlignmentTransform = InAlignmentTransform;
 
-	TArray<UARTrackedGeometry*> AllTrackedGeometries = GetAllTrackedGeometries();
+	TArray<UARTrackedGeometry*> AllTrackedGeometries = GetARCompositionComponent()->GetAllTrackedGeometries();
 	for (UARTrackedGeometry* TrackedGeometry : AllTrackedGeometries)
 	{
 		TrackedGeometry->UpdateAlignmentTransform(NewAlignmentTransform);
 	}
 
-	TArray<UARPin*> AllARPins = GetAllPins();
+	TArray<UARPin*> AllARPins = GetARCompositionComponent()->GetAllPins();
 	for (UARPin* SomePin : AllARPins)
 	{
 		SomePin->UpdateAlignmentTransform(NewAlignmentTransform);
 	}
-
-	SetAlignmentTransform_Internal(InAlignmentTransform);
 }
 
 TArray<FARTraceResult> FGoogleARCoreXRTrackingSystem::OnLineTraceTrackedObjects(const FVector2D ScreenCoord, EARLineTraceChannels TraceChannels)
@@ -297,8 +296,6 @@ void FGoogleARCoreXRTrackingSystem::OnRemovePin(UARPin* PinToRemove)
 
 void FGoogleARCoreXRTrackingSystem::AddReferencedObjects(FReferenceCollector& Collector)
 {
-	FARSystemBase::AddReferencedObjects(Collector);
-
 	if (LightEstimate != nullptr)
 	{
 		Collector.AddReferencedObject(LightEstimate);
diff --git a/Engine/Plugins/Runtime/GoogleARCore/Source/GoogleARCoreBase/Private/GoogleARCoreXRTrackingSystem.h b/Engine/Plugins/Runtime/GoogleARCore/Source/GoogleARCoreBase/Private/GoogleARCoreXRTrackingSystem.h
--- a/Engine/Plugins/Runtime/GoogleARCore/Source/GoogleARCoreBase/Private/GoogleARCoreXRTrackingSystem.h
+++ b/Engine/Plugins/Runtime/GoogleARCore/Source/GoogleARCoreBase/Private/GoogleARCoreXRTrackingSystem.h
@@ -11,7 +11,7 @@
 #include "ARSystem.h"
 #include "ARLightEstimate.h"
 
-class FGoogleARCoreXRTrackingSystem : public FARSystemBase
+class FGoogleARCoreXRTrackingSystem : public IARSystemSupport, public FXRTrackingSystemBase, public FGCObject, public TSharedFromThis<FGoogleARCoreXRTrackingSystem, ESPMode::ThreadSafe>
 {
 	friend class FGoogleARCoreXRCamera;
 
diff --git a/Engine/Plugins/Runtime/GoogleARCore/Source/GoogleARCoreBase/Public/GoogleARCoreAugmentedImage.h b/Engine/Plugins/Runtime/GoogleARCore/Source/GoogleARCoreBase/Public/GoogleARCoreAugmentedImage.h
--- a/Engine/Plugins/Runtime/GoogleARCore/Source/GoogleARCoreBase/Public/GoogleARCoreAugmentedImage.h
+++ b/Engine/Plugins/Runtime/GoogleARCore/Source/GoogleARCoreBase/Public/GoogleARCoreAugmentedImage.h
@@ -60,7 +60,7 @@ public:
 		float OutlineThickness, float PersistForSeconds = 0.0f) const override;
 
 	void UpdateTrackedGeometry(
-		const TSharedRef<FARSystemBase, ESPMode::ThreadSafe>& InTrackingSystem,
+		const TSharedRef<FARSupportInterface , ESPMode::ThreadSafe>& InTrackingSystem,
 		uint32 FrameNumber, double Timestamp, const FTransform& InLocalToTrackingTransform,
 		const FTransform& InAlignmentTransform,
 		const FVector &InCenter, const FVector &InExtent,
diff --git a/Engine/Plugins/Runtime/GoogleARCoreServices/Source/GoogleARCoreServices/Private/GoogleARCoreCloudARPinManager.cpp b/Engine/Plugins/Runtime/GoogleARCoreServices/Source/GoogleARCoreServices/Private/GoogleARCoreCloudARPinManager.cpp
--- a/Engine/Plugins/Runtime/GoogleARCoreServices/Source/GoogleARCoreServices/Private/GoogleARCoreCloudARPinManager.cpp
+++ b/Engine/Plugins/Runtime/GoogleARCoreServices/Source/GoogleARCoreServices/Private/GoogleARCoreCloudARPinManager.cpp
@@ -87,7 +87,7 @@ namespace {
 class FGoogleARCoreCloudARPinManagerAndroid : FGoogleARCoreCloudARPinManager
 {
 public:
-	FGoogleARCoreCloudARPinManagerAndroid(TSharedRef<FARSystemBase, ESPMode::ThreadSafe> InArSystem)
+	FGoogleARCoreCloudARPinManagerAndroid(TSharedRef<FARSupportInterface , ESPMode::ThreadSafe> InArSystem)
 		: FGoogleARCoreCloudARPinManager(InArSystem)
 	{
 
@@ -144,7 +144,7 @@ protected:
 class FGoogleARCoreCloudARPinManageriOS : FGoogleARCoreCloudARPinManager
 {
 public:
-	FGoogleARCoreCloudARPinManageriOS(TSharedRef<FARSystemBase, ESPMode::ThreadSafe> InArSystem)
+	FGoogleARCoreCloudARPinManageriOS(TSharedRef<FARSupportInterface , ESPMode::ThreadSafe> InArSystem)
 		: FGoogleARCoreCloudARPinManager(InArSystem)
 		, SessionHandle(nullptr)
 		, FrameHandle(nullptr)
@@ -249,7 +249,7 @@ private:
 };
 #endif // PLATFORM_IOS
 
-FGoogleARCoreCloudARPinManager* FGoogleARCoreCloudARPinManager::CreateCloudARPinManager(TSharedRef<FARSystemBase, ESPMode::ThreadSafe> InArSystem)
+FGoogleARCoreCloudARPinManager* FGoogleARCoreCloudARPinManager::CreateCloudARPinManager(TSharedRef<FARSupportInterface , ESPMode::ThreadSafe> InArSystem)
 {
 #if PLATFORM_ANDROID
 	return reinterpret_cast<FGoogleARCoreCloudARPinManager*>(new FGoogleARCoreCloudARPinManagerAndroid(InArSystem));
@@ -268,7 +268,7 @@ UCloudARPin* FGoogleARCoreCloudARPinManager::CreateAndHostCloudARPin(UARPin* Pin
 		return nullptr;
 	}
 
-	float WorldToMeterScale = ARSystem->GetWorldToMetersScale();
+	float WorldToMeterScale = ARSystem->GetXRTrackingSystem()->GetWorldToMetersScale();
 	ArSession* SessionHandle = GetSessionHandle();
 	ArAnchor* NewAnchorHandle = nullptr;
 #if PLATFORM_ANDROID
@@ -317,7 +317,7 @@ UCloudARPin* FGoogleARCoreCloudARPinManager::ResolveAndCreateCloudARPin(FString
 	ArSession* SessionHandle = GetSessionHandle();
 	ArAnchor* AnchorHandle = nullptr;
 	ArPose *PoseHandle = nullptr;
-	float WorldToMeterScale = ARSystem->GetWorldToMetersScale();
+	float WorldToMeterScale = ARSystem->GetXRTrackingSystem()->GetWorldToMetersScale();
 	ArPose_create(SessionHandle, nullptr, &PoseHandle);
 
 	OutTaskResult = ToCloudTaskResult(
@@ -389,7 +389,7 @@ void FGoogleARCoreCloudARPinManager::UpdateAllCloudARPins()
 		return;
 	}
 
-	float WorldToMeterScale = ARSystem->GetWorldToMetersScale();
+	float WorldToMeterScale = ARSystem->GetXRTrackingSystem()->GetWorldToMetersScale();
 
 	ArAnchorList* UpdatedAnchorListHandle = nullptr;
 	ArAnchorList_create(SessionHandle, &UpdatedAnchorListHandle);
diff --git a/Engine/Plugins/Runtime/GoogleARCoreServices/Source/GoogleARCoreServices/Private/GoogleARCoreCloudARPinManager.h b/Engine/Plugins/Runtime/GoogleARCoreServices/Source/GoogleARCoreServices/Private/GoogleARCoreCloudARPinManager.h
--- a/Engine/Plugins/Runtime/GoogleARCoreServices/Source/GoogleARCoreServices/Private/GoogleARCoreCloudARPinManager.h
+++ b/Engine/Plugins/Runtime/GoogleARCoreServices/Source/GoogleARCoreServices/Private/GoogleARCoreCloudARPinManager.h
@@ -13,7 +13,7 @@ class FGoogleARCoreCloudARPinManager : public FGCObject
 {
 
 public:
-	static FGoogleARCoreCloudARPinManager* CreateCloudARPinManager(TSharedRef<FARSystemBase, ESPMode::ThreadSafe> InArSystem);
+	static FGoogleARCoreCloudARPinManager* CreateCloudARPinManager(TSharedRef<FARSupportInterface , ESPMode::ThreadSafe> InArSystem);
 
 	// public member method:
 	virtual bool IsCloudARPinModeSupported(EARPinCloudMode NewMode) = 0;
@@ -43,7 +43,7 @@ public:
 
 protected:
 	// protected methods:
-	FGoogleARCoreCloudARPinManager(TSharedRef<FARSystemBase, ESPMode::ThreadSafe> InArSystem)
+	FGoogleARCoreCloudARPinManager(TSharedRef<FARSupportInterface , ESPMode::ThreadSafe> InArSystem)
 		: ARSystem(InArSystem)
 	{
 		ARSystem->OnAlignmentTransformUpdated.AddRaw(this, &FGoogleARCoreCloudARPinManager::OnAlignmentTransformUpdated);
@@ -68,7 +68,7 @@ protected:
 	}
 
 	// protected properties:
-	TSharedRef<FARSystemBase, ESPMode::ThreadSafe> ARSystem;
+	TSharedRef<FARSupportInterface , ESPMode::ThreadSafe> ARSystem;
 	TArray<UCloudARPin*> AllCloudARPins;
 	
 #if	ARCORE_SERVICE_SUPPORTED_PLATFORM
diff --git a/Engine/Plugins/Runtime/GoogleARCoreServices/Source/GoogleARCoreServices/Private/GoogleARCoreServicesManager.cpp b/Engine/Plugins/Runtime/GoogleARCoreServices/Source/GoogleARCoreServices/Private/GoogleARCoreServicesManager.cpp
--- a/Engine/Plugins/Runtime/GoogleARCoreServices/Source/GoogleARCoreServices/Private/GoogleARCoreServicesManager.cpp
+++ b/Engine/Plugins/Runtime/GoogleARCoreServices/Source/GoogleARCoreServices/Private/GoogleARCoreServicesManager.cpp
@@ -98,7 +98,7 @@ TArray<UCloudARPin*> FGoogleARCoreServicesManager::GetAllCloudARPin()
 
 bool FGoogleARCoreServicesManager::InitARSystem()
 {
-	ARSystem = StaticCastSharedPtr<FARSystemBase>(GEngine->XRSystem);
+	ARSystem = StaticCastSharedPtr<FXRTrackingSystemBase>(GEngine->XRSystem)->GetARCompositionComponent();
 
 	if (ARSystem.IsValid())
 	{
diff --git a/Engine/Plugins/Runtime/GoogleARCoreServices/Source/GoogleARCoreServices/Private/GoogleARCoreServicesManager.h b/Engine/Plugins/Runtime/GoogleARCoreServices/Source/GoogleARCoreServices/Private/GoogleARCoreServicesManager.h
--- a/Engine/Plugins/Runtime/GoogleARCoreServices/Source/GoogleARCoreServices/Private/GoogleARCoreServicesManager.h
+++ b/Engine/Plugins/Runtime/GoogleARCoreServices/Source/GoogleARCoreServices/Private/GoogleARCoreServicesManager.h
@@ -38,7 +38,7 @@ private:
 	bool bCloudARPinEnabled;
 	FGoogleARCoreServicesConfig CurrentServicesConfig;
 
-	TSharedPtr<FARSystemBase, ESPMode::ThreadSafe> ARSystem;
+	TSharedPtr<FARSupportInterface , ESPMode::ThreadSafe> ARSystem;
 	TUniquePtr<FGoogleARCoreCloudARPinManager> CloudARPinManager;
 };
 
diff --git a/Engine/Plugins/Runtime/GoogleVR/GoogleVRHMD/Source/GoogleVRHMD/Private/GoogleVRHMD.cpp b/Engine/Plugins/Runtime/GoogleVR/GoogleVRHMD/Source/GoogleVRHMD/Private/GoogleVRHMD.cpp
--- a/Engine/Plugins/Runtime/GoogleVR/GoogleVRHMD/Source/GoogleVRHMD/Private/GoogleVRHMD.cpp
+++ b/Engine/Plugins/Runtime/GoogleVR/GoogleVRHMD/Source/GoogleVRHMD/Private/GoogleVRHMD.cpp
@@ -307,6 +307,7 @@ TSharedPtr< class IXRTrackingSystem, ESPMode::ThreadSafe > FGoogleVRHMDPlugin::C
 /////////////////////////////////////
 FGoogleVRHMD::FGoogleVRHMD(const FAutoRegister& AutoRegister)
 	: FSceneViewExtensionBase(AutoRegister)
+	, FHeadMountedDisplayBase(nullptr)
 #if GOOGLEVRHMD_SUPPORTED_PLATFORMS
 	, CustomPresent(nullptr)
 #endif
diff --git a/Engine/Plugins/Runtime/Oculus/OculusVR/Source/OculusHMD/OculusHMD.Build.cs b/Engine/Plugins/Runtime/Oculus/OculusVR/Source/OculusHMD/OculusHMD.Build.cs
--- a/Engine/Plugins/Runtime/Oculus/OculusVR/Source/OculusHMD/OculusHMD.Build.cs
+++ b/Engine/Plugins/Runtime/Oculus/OculusVR/Source/OculusHMD/OculusHMD.Build.cs
@@ -45,7 +45,6 @@ namespace UnrealBuildTool.Rules
 					"RenderCore",
 					"Renderer",
 					"ShaderCore",
-					"HeadMountedDisplay",
 					"Slate",
 					"SlateCore",
 					"ImageWrapper",
@@ -58,8 +57,14 @@ namespace UnrealBuildTool.Rules
 					"ProceduralMeshComponent",
                     "Projects",
                 });
+            PublicDependencyModuleNames.AddRange(
+                new string[]
+                {
+                    "HeadMountedDisplay",
+                });
+
 
-			if (Target.bBuildEditor == true)
+            if (Target.bBuildEditor == true)
 			{
 				PrivateDependencyModuleNames.Add("UnrealEd");
 			}
diff --git a/Engine/Plugins/Runtime/Oculus/OculusVR/Source/OculusHMD/Private/OculusHMD.cpp b/Engine/Plugins/Runtime/Oculus/OculusVR/Source/OculusHMD/Private/OculusHMD.cpp
--- a/Engine/Plugins/Runtime/Oculus/OculusVR/Source/OculusHMD/Private/OculusHMD.cpp
+++ b/Engine/Plugins/Runtime/Oculus/OculusVR/Source/OculusHMD/Private/OculusHMD.cpp
@@ -1730,6 +1730,7 @@ namespace OculusHMD
 
 	FOculusHMD::FOculusHMD(const FAutoRegister& AutoRegister)
 		: FSceneViewExtensionBase(AutoRegister)
+		, FHeadMountedDisplayBase(nullptr)
 		, ConsoleCommands(this)
 	{
 		Flags.Raw = 0;
diff --git a/Engine/Plugins/Runtime/Steam/SteamVR/Source/SteamVR/Private/SteamVRHMD.cpp b/Engine/Plugins/Runtime/Steam/SteamVR/Source/SteamVR/Private/SteamVRHMD.cpp
--- a/Engine/Plugins/Runtime/Steam/SteamVR/Source/SteamVR/Private/SteamVRHMD.cpp
+++ b/Engine/Plugins/Runtime/Steam/SteamVR/Source/SteamVR/Private/SteamVRHMD.cpp
@@ -1494,6 +1494,7 @@ bool FSteamVRHMD::NeedReAllocateViewportRenderTarget(const FViewport& Viewport)
 }
 
 FSteamVRHMD::FSteamVRHMD(ISteamVRPlugin* InSteamVRPlugin) :
+	FHeadMountedDisplayBase(nullptr),
 	bHmdEnabled(true),
 	HmdWornState(EHMDWornState::Unknown),
 	bStereoDesired(false),
diff --git a/Engine/Plugins/Runtime/WindowsMixedReality/Source/WindowsMixedRealityHMD/Private/WindowsMixedRealityHMD.cpp b/Engine/Plugins/Runtime/WindowsMixedReality/Source/WindowsMixedRealityHMD/Private/WindowsMixedRealityHMD.cpp
--- a/Engine/Plugins/Runtime/WindowsMixedReality/Source/WindowsMixedRealityHMD/Private/WindowsMixedRealityHMD.cpp
+++ b/Engine/Plugins/Runtime/WindowsMixedReality/Source/WindowsMixedRealityHMD/Private/WindowsMixedRealityHMD.cpp
@@ -966,6 +966,7 @@ namespace WindowsMixedReality
 
 	FWindowsMixedRealityHMD::FWindowsMixedRealityHMD(const FAutoRegister& AutoRegister)
 		: FSceneViewExtensionBase(AutoRegister)
+		, FHeadMountedDisplayBase(nullptr)
 		, mCustomPresent(nullptr)
 		, CurrOrientation(FQuat::Identity)
 		, CurrPosition(FVector::ZeroVector)
diff --git a/Engine/Source/Runtime/AugmentedReality/AugmentedReality.build.cs b/Engine/Source/Runtime/AugmentedReality/AugmentedReality.build.cs
--- a/Engine/Source/Runtime/AugmentedReality/AugmentedReality.build.cs
+++ b/Engine/Source/Runtime/AugmentedReality/AugmentedReality.build.cs
@@ -19,9 +19,17 @@ namespace UnrealBuildTool.Rules
                     "RenderCore",
 					"ShaderCore",
 					"RHI",
-                    "HeadMountedDisplay"
+                    //"HeadMountedDisplay"
 				}
 			);
-		}
+
+            PublicIncludePathModuleNames.AddRange(
+                new string[]
+                {
+					"HeadMountedDisplay",
+                }
+           );
+
+        }
 	}
 }
diff --git a/Engine/Source/Runtime/AugmentedReality/Private/ARBlueprintLibrary.cpp b/Engine/Source/Runtime/AugmentedReality/Private/ARBlueprintLibrary.cpp
--- a/Engine/Source/Runtime/AugmentedReality/Private/ARBlueprintLibrary.cpp
+++ b/Engine/Source/Runtime/AugmentedReality/Private/ARBlueprintLibrary.cpp
@@ -8,7 +8,7 @@
 #include "ARPin.h"
 #include "ARTrackable.h"
 
-TSharedPtr<FARSystemBase, ESPMode::ThreadSafe> UARBlueprintLibrary::RegisteredARSystem = nullptr;
+TSharedPtr<FARSupportInterface , ESPMode::ThreadSafe> UARBlueprintLibrary::RegisteredARSystem = nullptr;
 
 EARTrackingQuality UARBlueprintLibrary::GetTrackingQuality()
 {
@@ -235,13 +235,13 @@ void UARBlueprintLibrary::RemovePin( UARPin* PinToRemove )
 
 
 
-void UARBlueprintLibrary::RegisterAsARSystem(const TSharedPtr<FARSystemBase, ESPMode::ThreadSafe>& NewARSystem)
+void UARBlueprintLibrary::RegisterAsARSystem(const TSharedPtr<FARSupportInterface , ESPMode::ThreadSafe>& NewARSystem)
 {
 	RegisteredARSystem = NewARSystem;
 }
 
 
-const TSharedPtr<FARSystemBase, ESPMode::ThreadSafe>& UARBlueprintLibrary::GetARSystem()
+const TSharedPtr<FARSupportInterface , ESPMode::ThreadSafe>& UARBlueprintLibrary::GetARSystem()
 {
 	return RegisteredARSystem;
 }
diff --git a/Engine/Source/Runtime/AugmentedReality/Private/ARBlueprintProxy.cpp b/Engine/Source/Runtime/AugmentedReality/Private/ARBlueprintProxy.cpp
--- a/Engine/Source/Runtime/AugmentedReality/Private/ARBlueprintProxy.cpp
+++ b/Engine/Source/Runtime/AugmentedReality/Private/ARBlueprintProxy.cpp
@@ -3,7 +3,7 @@
 #include "ARBlueprintProxy.h"
 #include "ARSystem.h"
 
-TSharedPtr<FARSystemBase, ESPMode::ThreadSafe> UARBaseAsyncTaskBlueprintProxy::RegisteredARSystem = nullptr;
+TSharedPtr<FARSupportInterface , ESPMode::ThreadSafe> UARBaseAsyncTaskBlueprintProxy::RegisteredARSystem = nullptr;
 
 UARBaseAsyncTaskBlueprintProxy::UARBaseAsyncTaskBlueprintProxy(const FObjectInitializer& ObjectInitializer)
 	: Super(ObjectInitializer)
@@ -35,13 +35,13 @@ void UARBaseAsyncTaskBlueprintProxy::Tick(float DeltaTime)
 	}
 }
 
-void UARBaseAsyncTaskBlueprintProxy::RegisterAsARSystem(const TSharedPtr<FARSystemBase, ESPMode::ThreadSafe>& NewARSystem)
+void UARBaseAsyncTaskBlueprintProxy::RegisterAsARSystem(const TSharedPtr<FARSupportInterface , ESPMode::ThreadSafe>& NewARSystem)
 {
 	RegisteredARSystem = NewARSystem;
 }
 
 
-const TSharedPtr<FARSystemBase, ESPMode::ThreadSafe>& UARBaseAsyncTaskBlueprintProxy::GetARSystem()
+const TSharedPtr<FARSupportInterface , ESPMode::ThreadSafe>& UARBaseAsyncTaskBlueprintProxy::GetARSystem()
 {
 	return RegisteredARSystem;
 }
diff --git a/Engine/Source/Runtime/AugmentedReality/Private/ARPin.cpp b/Engine/Source/Runtime/AugmentedReality/Private/ARPin.cpp
--- a/Engine/Source/Runtime/AugmentedReality/Private/ARPin.cpp
+++ b/Engine/Source/Runtime/AugmentedReality/Private/ARPin.cpp
@@ -14,7 +14,7 @@
 
 uint32 UARPin::DebugPinId = 0;
 
-void UARPin::InitARPin( const TSharedRef<FARSystemBase, ESPMode::ThreadSafe>& InTrackingSystemOwner, USceneComponent* InComponentToPin, const FTransform& InLocalToTrackingTransform, UARTrackedGeometry* InTrackedGeometry, const FName InDebugName )
+void UARPin::InitARPin( const TSharedRef<FARSupportInterface , ESPMode::ThreadSafe>& InTrackingSystemOwner, USceneComponent* InComponentToPin, const FTransform& InLocalToTrackingTransform, UARTrackedGeometry* InTrackedGeometry, const FName InDebugName )
 {
 	ARSystem = InTrackingSystemOwner;
 	LocalToTrackingTransform = InLocalToTrackingTransform;
@@ -42,7 +42,7 @@ FTransform UARPin::GetLocalToTrackingTransform() const
 
 FTransform UARPin::GetLocalToWorldTransform() const
 {
-	return GetLocalToTrackingTransform() * GetARSystem()->GetTrackingToWorldTransform();
+	return GetLocalToTrackingTransform() * GetARSystem()->GetXRTrackingSystem()->GetTrackingToWorldTransform();
 }
 
 EARTrackingState UARPin::GetTrackingState() const
@@ -144,7 +144,7 @@ void UARPin::SetOnARTransformUpdated( const FOnARTransformUpdated& InHandler )
 	OnARTransformUpdated = InHandler;
 }
 
-TSharedPtr<FARSystemBase, ESPMode::ThreadSafe> UARPin::GetARSystem() const
+TSharedPtr<FARSupportInterface , ESPMode::ThreadSafe> UARPin::GetARSystem() const
 {
 	auto MyARSystem = ARSystem.Pin();
 	return MyARSystem;
diff --git a/Engine/Source/Runtime/AugmentedReality/Private/ARSupportInterface.cpp b/Engine/Source/Runtime/AugmentedReality/Private/ARSupportInterface.cpp
new file mode 100644
--- /dev/null
+++ b/Engine/Source/Runtime/AugmentedReality/Private/ARSupportInterface.cpp
@@ -0,0 +1,307 @@
+// Copyright 1998-2018 Epic Games, Inc. All Rights Reserved.
+
+#include "ARSupportInterface.h"
+#include "ArTraceResult.h"
+#include "ARSystem.h"
+#include "Features/IModularFeatures.h"
+#include "ARBlueprintLibrary.h"
+#include "ARBlueprintProxy.h"
+#include "Templates/SharedPointer.h"
+
+FARSupportInterface ::FARSupportInterface (IARSystemSupport* InARImplementation, IXRTrackingSystem* InXRTrackingSystem)
+	: ARImplemention(InARImplementation)
+	, XRTrackingSystem(InXRTrackingSystem)
+	, AlignmentTransform(FTransform::Identity)
+	, ARSettings(NewObject<UARSessionConfig>())
+{
+}
+
+
+FARSupportInterface ::~FARSupportInterface ()
+{
+	IModularFeatures::Get().UnregisterModularFeature(FARSupportInterface ::GetModularFeatureName(), this);
+
+	UARBlueprintLibrary::RegisterAsARSystem(nullptr);
+	UARBaseAsyncTaskBlueprintProxy::RegisterAsARSystem(nullptr);
+}
+
+
+void FARSupportInterface ::InitializeARSystem()
+{
+	// Register our ability to support Unreal AR API.
+	IModularFeatures::Get().RegisterModularFeature(FARSupportInterface ::GetModularFeatureName(), this);
+
+	if (ARImplemention)
+	{
+		UARBlueprintLibrary::RegisterAsARSystem(MakeShareable(this));
+		UARBaseAsyncTaskBlueprintProxy::RegisterAsARSystem(MakeShareable(this));
+
+		ARImplemention->OnARSystemInitialized();
+	}
+}
+
+IXRTrackingSystem* FARSupportInterface ::GetXRTrackingSystem()
+{
+	return XRTrackingSystem;
+}
+
+const FTransform& FARSupportInterface ::GetAlignmentTransform() const
+{
+	return AlignmentTransform;
+}
+
+const UARSessionConfig& FARSupportInterface ::GetSessionConfig() const
+{
+	check(ARSettings != nullptr);
+	return *ARSettings;
+}
+
+UARSessionConfig& FARSupportInterface ::AccessSessionConfig()
+{
+	check(ARSettings != nullptr);
+	return *ARSettings;
+}
+
+bool FARSupportInterface ::StartARGameFrame(FWorldContext& WorldContext)
+{
+	if (ARImplemention)
+	{
+		return ARImplemention->OnStartARGameFrame(WorldContext);
+	}
+	return false;
+}
+
+
+EARTrackingQuality FARSupportInterface ::GetTrackingQuality() const
+{
+	if (ARImplemention)
+	{
+		return ARImplemention->OnGetTrackingQuality();
+	}
+	return EARTrackingQuality::NotTracking;
+}
+
+void FARSupportInterface ::StartARSession(UARSessionConfig* InSessionConfig)
+{
+	if (ARImplemention)
+	{
+		ARSettings = InSessionConfig;
+		ARImplemention->OnStartARSession(InSessionConfig);
+	}
+}
+
+void FARSupportInterface ::PauseARSession()
+{
+	if (ARImplemention)
+	{
+		if (GetARSessionStatus().Status == EARSessionStatus::Running)
+		{
+			ARImplemention->OnPauseARSession();
+		}
+	}
+}
+
+void FARSupportInterface ::StopARSession()
+{
+	if (ARImplemention)
+	{
+		if (GetARSessionStatus().Status == EARSessionStatus::Running)
+		{
+			ARImplemention->OnStopARSession();
+		}
+	}
+}
+
+FARSessionStatus FARSupportInterface ::GetARSessionStatus() const
+{
+	if (ARImplemention)
+	{
+		return ARImplemention->OnGetARSessionStatus();
+	}
+	return EARSessionStatus::NotSupported;
+}
+
+bool FARSupportInterface ::IsSessionTypeSupported(EARSessionType SessionType) const
+{
+	if (ARImplemention)
+	{
+		return ARImplemention->OnIsTrackingTypeSupported(SessionType);
+	}
+	return false;
+}
+
+void FARSupportInterface ::SetAlignmentTransform(const FTransform& InAlignmentTransform)
+{
+	if (ARImplemention)
+	{
+		ARImplemention->OnSetAlignmentTransform(InAlignmentTransform);
+	}
+	AlignmentTransform = InAlignmentTransform;
+	OnAlignmentTransformUpdated.Broadcast(InAlignmentTransform);
+}
+
+TArray<FARTraceResult> FARSupportInterface ::LineTraceTrackedObjects(const FVector2D ScreenCoord, EARLineTraceChannels TraceChannels)
+{
+	if (ARImplemention)
+	{
+		return ARImplemention->OnLineTraceTrackedObjects(ScreenCoord, TraceChannels);
+	}
+	return TArray<FARTraceResult>();
+}
+
+TArray<UARTrackedGeometry*> FARSupportInterface ::GetAllTrackedGeometries() const
+{
+	if (ARImplemention)
+	{
+		return ARImplemention->OnGetAllTrackedGeometries();
+	}
+	return TArray<UARTrackedGeometry*>();
+}
+
+TArray<UARPin*> FARSupportInterface ::GetAllPins() const
+{
+	if (ARImplemention)
+	{
+		return ARImplemention->OnGetAllPins();
+	}
+	return TArray<UARPin*>();
+}
+
+UARTextureCameraImage* FARSupportInterface ::GetCameraImage()
+{
+	if (ARImplemention)
+	{
+		return ARImplemention->OnGetCameraImage();
+	}
+	return nullptr;
+}
+
+UARTextureCameraDepth* FARSupportInterface ::GetCameraDepth()
+{
+	if (ARImplemention)
+	{
+		return ARImplemention->OnGetCameraDepth();
+	}
+	return nullptr;
+}
+
+bool FARSupportInterface ::AddManualEnvironmentCaptureProbe(FVector Location, FVector Extent)
+{
+	if (ARImplemention)
+	{
+		return ARImplemention->OnAddManualEnvironmentCaptureProbe(Location, Extent);
+	}
+	return false;
+}
+
+TSharedPtr<FARGetCandidateObjectAsyncTask, ESPMode::ThreadSafe> FARSupportInterface ::GetCandidateObject(FVector Location, FVector Extent) const
+{
+	if (ARImplemention)
+	{
+		return ARImplemention->OnGetCandidateObject(Location, Extent);
+	}
+	return TSharedPtr<FARGetCandidateObjectAsyncTask, ESPMode::ThreadSafe>();
+}
+
+TSharedPtr<FARSaveWorldAsyncTask, ESPMode::ThreadSafe> FARSupportInterface ::SaveWorld() const
+{
+	if (ARImplemention)
+	{
+		return ARImplemention->OnSaveWorld();
+	}
+	return TSharedPtr<FARSaveWorldAsyncTask, ESPMode::ThreadSafe>();
+}
+
+EARWorldMappingState FARSupportInterface ::GetWorldMappingStatus() const
+{
+	if (ARImplemention)
+	{
+		return ARImplemention->OnGetWorldMappingStatus();
+	}
+	return EARWorldMappingState::NotAvailable;
+}
+
+
+UARLightEstimate* FARSupportInterface ::GetCurrentLightEstimate() const
+{
+	if (ARImplemention)
+	{
+		return ARImplemention->OnGetCurrentLightEstimate();
+	}
+	return nullptr;
+}
+
+UARPin* FARSupportInterface ::PinComponent(USceneComponent* ComponentToPin, const FTransform& PinToWorldTransform, UARTrackedGeometry* TrackedGeometry, const FName DebugName)
+{
+	if (ARImplemention)
+	{
+		return ARImplemention->OnPinComponent(ComponentToPin, PinToWorldTransform, TrackedGeometry, DebugName);
+	}
+	return nullptr;
+}
+
+UARPin* FARSupportInterface ::PinComponent(USceneComponent* ComponentToPin, const FARTraceResult& HitResult, const FName DebugName)
+{
+	if (ARImplemention)
+	{
+		return ARImplemention->OnPinComponent(ComponentToPin, HitResult.GetLocalToWorldTransform(), HitResult.GetTrackedGeometry(), DebugName);
+	}
+	return nullptr;
+}
+
+void FARSupportInterface ::RemovePin(UARPin* PinToRemove)
+{
+	if (ARImplemention)
+	{
+		ARImplemention->OnRemovePin(PinToRemove);
+	}
+}
+
+TArray<FARVideoFormat> FARSupportInterface ::GetSupportedVideoFormats(EARSessionType SessionType) const
+{
+	if (ARImplemention)
+	{
+		return ARImplemention->OnGetSupportedVideoFormats(SessionType);
+	}
+	return TArray<FARVideoFormat>();
+}
+
+/** @return the current point cloud data for the ar scene */
+TArray<FVector> FARSupportInterface ::GetPointCloud() const
+{
+	if (ARImplemention)
+	{
+		return ARImplemention->OnGetPointCloud();
+	}
+	return TArray<FVector>();
+}
+
+
+
+void* FARSupportInterface ::GetARSessionRawPointer()
+{
+	if (ARImplemention)
+	{
+		return ARImplemention->GetARSessionRawPointer();
+	}
+	return nullptr;
+}
+
+void* FARSupportInterface ::GetGameThreadARFrameRawPointer()
+{
+	if (ARImplemention)
+	{
+		return ARImplemention->GetGameThreadARFrameRawPointer();
+	}
+	return nullptr;
+}
+
+
+void FARSupportInterface ::AddReferencedObjects(FReferenceCollector& Collector)
+{
+	if (ARSettings != nullptr)
+	{
+		Collector.AddReferencedObject(ARSettings);
+	}
+}
+
diff --git a/Engine/Source/Runtime/AugmentedReality/Private/ARSystem.cpp b/Engine/Source/Runtime/AugmentedReality/Private/ARSystem.cpp
deleted file mode 100644
diff --git a/Engine/Source/Runtime/AugmentedReality/Private/ARTraceResult.cpp b/Engine/Source/Runtime/AugmentedReality/Private/ARTraceResult.cpp
--- a/Engine/Source/Runtime/AugmentedReality/Private/ARTraceResult.cpp
+++ b/Engine/Source/Runtime/AugmentedReality/Private/ARTraceResult.cpp
@@ -14,7 +14,7 @@ FARTraceResult::FARTraceResult()
 }
 
 
-FARTraceResult::FARTraceResult( const TSharedPtr<FARSystemBase, ESPMode::ThreadSafe>& InARSystem, float InDistanceFromCamera, EARLineTraceChannels InTraceChannel, const FTransform& InLocalToTrackingTransform, UARTrackedGeometry* InTrackedGeometry )
+FARTraceResult::FARTraceResult( const TSharedPtr<FARSupportInterface , ESPMode::ThreadSafe>& InARSystem, float InDistanceFromCamera, EARLineTraceChannels InTraceChannel, const FTransform& InLocalToTrackingTransform, UARTrackedGeometry* InTrackedGeometry )
 : DistanceFromCamera(InDistanceFromCamera)
 , TraceChannel(InTraceChannel)
 , LocalToTrackingTransform(InLocalToTrackingTransform)
@@ -37,7 +37,7 @@ FTransform FARTraceResult::GetLocalToTrackingTransform() const
 
 FTransform FARTraceResult::GetLocalToWorldTransform() const
 {
-	return LocalToTrackingTransform * ARSystem->GetTrackingToWorldTransform();
+	return LocalToTrackingTransform * ARSystem->GetXRTrackingSystem()->GetTrackingToWorldTransform();
 }
 
 
diff --git a/Engine/Source/Runtime/AugmentedReality/Private/ARTrackable.cpp b/Engine/Source/Runtime/AugmentedReality/Private/ARTrackable.cpp
--- a/Engine/Source/Runtime/AugmentedReality/Private/ARTrackable.cpp
+++ b/Engine/Source/Runtime/AugmentedReality/Private/ARTrackable.cpp
@@ -29,7 +29,7 @@ void UARTrackedGeometry::DebugDraw( UWorld* World, const FLinearColor& OutlineCo
 	DrawDebugCoordinateSystem(World, Location, Rotation, Scale3D.X, true, PersistForSeconds, 0, OutlineThickness);
 }
 
-TSharedPtr<FARSystemBase, ESPMode::ThreadSafe> UARTrackedGeometry::GetARSystem() const
+TSharedPtr<FARSupportInterface , ESPMode::ThreadSafe> UARTrackedGeometry::GetARSystem() const
 {
 	auto MyARSystem = ARSystem.Pin();
 	return MyARSystem;
@@ -62,7 +62,7 @@ void UARTrackedGeometry::SetTrackingState(EARTrackingState NewState)
 
 FTransform UARTrackedGeometry::GetLocalToWorldTransform() const
 {
-	return GetLocalToTrackingTransform() * GetARSystem()->GetTrackingToWorldTransform();
+	return GetLocalToTrackingTransform() * GetARSystem()->GetXRTrackingSystem()->GetTrackingToWorldTransform();
 }
 
 int32 UARTrackedGeometry::GetLastUpdateFrameNumber() const
@@ -80,7 +80,7 @@ float UARTrackedGeometry::GetLastUpdateTimestamp() const
 	return LastUpdateTimestamp;
 }
 
-void UARTrackedGeometry::UpdateTrackedGeometry(const TSharedRef<FARSystemBase, ESPMode::ThreadSafe>& InTrackingSystem, uint32 FrameNumber, double Timestamp, const FTransform& InLocalToTrackingTransform, const FTransform& InAlignmentTransform )
+void UARTrackedGeometry::UpdateTrackedGeometry(const TSharedRef<FARSupportInterface , ESPMode::ThreadSafe>& InTrackingSystem, uint32 FrameNumber, double Timestamp, const FTransform& InLocalToTrackingTransform, const FTransform& InAlignmentTransform )
 {
 	ARSystem = InTrackingSystem;
 	LocalToTrackingTransform = InLocalToTrackingTransform;
@@ -118,7 +118,7 @@ IARRef* UARTrackedGeometry::GetNativeResource()
 //
 //
 //
-void UARPlaneGeometry::UpdateTrackedGeometry(const TSharedRef<FARSystemBase, ESPMode::ThreadSafe>& InTrackingSystem, uint32 FrameNumber, double Timestamp, const FTransform& InLocalToTrackingTransform, const FTransform& InAlignmentTransform, const FVector InCenter, const FVector InExtent )
+void UARPlaneGeometry::UpdateTrackedGeometry(const TSharedRef<FARSupportInterface , ESPMode::ThreadSafe>& InTrackingSystem, uint32 FrameNumber, double Timestamp, const FTransform& InLocalToTrackingTransform, const FTransform& InAlignmentTransform, const FVector InCenter, const FVector InExtent )
 {
 	Super::UpdateTrackedGeometry(InTrackingSystem, FrameNumber, Timestamp, InLocalToTrackingTransform, InAlignmentTransform);
 	Center = InCenter;
@@ -133,7 +133,7 @@ void UARPlaneGeometry::UpdateTrackedGeometry(const TSharedRef<FARSystemBase, ESP
 	SubsumedBy = nullptr;
 }
 
-void UARPlaneGeometry::UpdateTrackedGeometry(const TSharedRef<FARSystemBase, ESPMode::ThreadSafe>& InTrackingSystem, uint32 FrameNumber, double Timestamp, const FTransform& InLocalToTrackingTransform, const FTransform& InAlignmentTransform, const FVector InCenter, const FVector InExtent, const TArray<FVector>& InBoundingPoly, UARPlaneGeometry* InSubsumedBy)
+void UARPlaneGeometry::UpdateTrackedGeometry(const TSharedRef<FARSupportInterface , ESPMode::ThreadSafe>& InTrackingSystem, uint32 FrameNumber, double Timestamp, const FTransform& InLocalToTrackingTransform, const FTransform& InAlignmentTransform, const FVector InCenter, const FVector InExtent, const TArray<FVector>& InBoundingPoly, UARPlaneGeometry* InSubsumedBy)
 {
 	Super::UpdateTrackedGeometry(InTrackingSystem, FrameNumber, Timestamp, InLocalToTrackingTransform, InAlignmentTransform);
 	Center = InCenter;
@@ -181,13 +181,13 @@ void UARTrackedImage::DebugDraw(UWorld* World, const FLinearColor& OutlineColor,
 	ARDebugHelpers::DrawDebugString(World, WorldSpaceCenter, CurAnchorDebugName, 0.25f * OutlineThickness, OutlineRGB, PersistForSeconds, true);
 }
 
-void UARTrackedImage::UpdateTrackedGeometry(const TSharedRef<FARSystemBase, ESPMode::ThreadSafe>& InTrackingSystem, uint32 FrameNumber, double Timestamp, const FTransform& InLocalToTrackingTransform, const FTransform& InAlignmentTransform, UARCandidateImage* InDetectedImage)
+void UARTrackedImage::UpdateTrackedGeometry(const TSharedRef<FARSupportInterface , ESPMode::ThreadSafe>& InTrackingSystem, uint32 FrameNumber, double Timestamp, const FTransform& InLocalToTrackingTransform, const FTransform& InAlignmentTransform, UARCandidateImage* InDetectedImage)
 {
 	Super::UpdateTrackedGeometry(InTrackingSystem, FrameNumber, Timestamp, InLocalToTrackingTransform, InAlignmentTransform);
 	DetectedImage = InDetectedImage;
 }
 
-void UARFaceGeometry::UpdateFaceGeometry(const TSharedRef<FARSystemBase, ESPMode::ThreadSafe>& InTrackingSystem, uint32 FrameNumber, double Timestamp, const FTransform& InLocalToTrackingTransform, const FTransform& InAlignmentTransform, FARBlendShapeMap& InBlendShapes, TArray<FVector>& InVertices, const TArray<int32>& Indices, const FTransform& InLeftEyeTransform, const FTransform& InRightEyeTransform, const FVector& InLookAtTarget)
+void UARFaceGeometry::UpdateFaceGeometry(const TSharedRef<FARSupportInterface , ESPMode::ThreadSafe>& InTrackingSystem, uint32 FrameNumber, double Timestamp, const FTransform& InLocalToTrackingTransform, const FTransform& InAlignmentTransform, FARBlendShapeMap& InBlendShapes, TArray<FVector>& InVertices, const TArray<int32>& Indices, const FTransform& InLeftEyeTransform, const FTransform& InRightEyeTransform, const FVector& InLookAtTarget)
 {
 	Super::UpdateTrackedGeometry(InTrackingSystem, FrameNumber, Timestamp, InLocalToTrackingTransform, InAlignmentTransform);
 	BlendShapes = MoveTemp(InBlendShapes);
@@ -213,7 +213,7 @@ void UARTrackedPoint::DebugDraw(UWorld* World, const FLinearColor& OutlineColor,
 	DrawDebugPoint(World, LocalToWorldTransform.GetLocation(), 0.5f, OutlineRGB, false, PersistForSeconds, 0);
 }
 
-void UARTrackedPoint::UpdateTrackedGeometry(const TSharedRef<FARSystemBase, ESPMode::ThreadSafe>& InTrackingSystem, uint32 FrameNumber, double Timestamp, const FTransform& InLocalToTrackingTransform, const FTransform& InAlignmentTransform)
+void UARTrackedPoint::UpdateTrackedGeometry(const TSharedRef<FARSupportInterface , ESPMode::ThreadSafe>& InTrackingSystem, uint32 FrameNumber, double Timestamp, const FTransform& InLocalToTrackingTransform, const FTransform& InAlignmentTransform)
 {
 	Super::UpdateTrackedGeometry(InTrackingSystem, FrameNumber, Timestamp, InLocalToTrackingTransform, InAlignmentTransform);
 }
@@ -273,7 +273,7 @@ void UAREnvironmentCaptureProbe::DebugDraw(UWorld* World, const FLinearColor& Ou
 	DrawDebugBox( World, LocalToWorldTransform.GetLocation(), Extent, LocalToWorldTransform.GetRotation(), OutlineRGB, false, PersistForSeconds, 0, 0.1f * OutlineThickness );
 }
 
-void UAREnvironmentCaptureProbe::UpdateEnvironmentCapture(const TSharedRef<FARSystemBase, ESPMode::ThreadSafe>& InTrackingSystem, uint32 InFrameNumber, double InTimestamp, const FTransform& InLocalToTrackingTransform, const FTransform& InAlignmentTransform, FVector InExtent)
+void UAREnvironmentCaptureProbe::UpdateEnvironmentCapture(const TSharedRef<FARSupportInterface , ESPMode::ThreadSafe>& InTrackingSystem, uint32 InFrameNumber, double InTimestamp, const FTransform& InLocalToTrackingTransform, const FTransform& InAlignmentTransform, FVector InExtent)
 {
 	Super::UpdateTrackedGeometry(InTrackingSystem, InFrameNumber, InTimestamp, InLocalToTrackingTransform, InAlignmentTransform);
 
@@ -300,7 +300,7 @@ void UARTrackedObject::DebugDraw(UWorld* World, const FLinearColor& OutlineColor
 	DrawDebugPoint(World, LocalToWorldTransform.GetLocation(), 0.5f, OutlineRGB, false, PersistForSeconds, 0);
 }
 
-void UARTrackedObject::UpdateTrackedGeometry(const TSharedRef<FARSystemBase, ESPMode::ThreadSafe>& InTrackingSystem, uint32 FrameNumber, double Timestamp, const FTransform& InLocalToTrackingTransform, const FTransform& InAlignmentTransform, UARCandidateObject* InDetectedObject)
+void UARTrackedObject::UpdateTrackedGeometry(const TSharedRef<FARSupportInterface , ESPMode::ThreadSafe>& InTrackingSystem, uint32 FrameNumber, double Timestamp, const FTransform& InLocalToTrackingTransform, const FTransform& InAlignmentTransform, UARCandidateObject* InDetectedObject)
 {
 	Super::UpdateTrackedGeometry(InTrackingSystem, FrameNumber, Timestamp, InLocalToTrackingTransform, InAlignmentTransform);
 	DetectedObject = InDetectedObject;
diff --git a/Engine/Source/Runtime/AugmentedReality/Public/ARBlueprintLibrary.h b/Engine/Source/Runtime/AugmentedReality/Public/ARBlueprintLibrary.h
--- a/Engine/Source/Runtime/AugmentedReality/Public/ARBlueprintLibrary.h
+++ b/Engine/Source/Runtime/AugmentedReality/Public/ARBlueprintLibrary.h
@@ -9,6 +9,7 @@
 #include "ARTraceResult.h"
 #include "ARSessionConfig.h"
 #include "ARTrackable.h"
+#include "ARSupportInterface.h"
 #include "ARBlueprintLibrary.generated.h"
 
 
@@ -194,11 +195,11 @@ public:
 	static TSharedPtr<FARGetCandidateObjectAsyncTask, ESPMode::ThreadSafe> GetCandidateObject(FVector Location, FVector Extent);
 	
 public:
-	static void RegisterAsARSystem(const TSharedPtr<FARSystemBase, ESPMode::ThreadSafe>& NewArSystem);
+	static void RegisterAsARSystem(const TSharedPtr<FARSupportInterface , ESPMode::ThreadSafe>& NewArSystem);
 	
 private:
-	static const TSharedPtr<FARSystemBase, ESPMode::ThreadSafe>& GetARSystem();
-	static TSharedPtr<FARSystemBase, ESPMode::ThreadSafe> RegisteredARSystem;
+	static const TSharedPtr<FARSupportInterface , ESPMode::ThreadSafe>& GetARSystem();
+	static TSharedPtr<FARSupportInterface , ESPMode::ThreadSafe> RegisteredARSystem;
 };
 
 
diff --git a/Engine/Source/Runtime/AugmentedReality/Public/ARBlueprintProxy.h b/Engine/Source/Runtime/AugmentedReality/Public/ARBlueprintProxy.h
--- a/Engine/Source/Runtime/AugmentedReality/Public/ARBlueprintProxy.h
+++ b/Engine/Source/Runtime/AugmentedReality/Public/ARBlueprintProxy.h
@@ -10,11 +10,12 @@
 #include "Kismet/BlueprintAsyncActionBase.h"
 
 #include "ARTypes.h"
+#include "ARSupportInterface.h"
 
 #include "ARBlueprintProxy.generated.h"
 
 UCLASS(Abstract)
-class UARBaseAsyncTaskBlueprintProxy :
+class AUGMENTEDREALITY_API UARBaseAsyncTaskBlueprintProxy :
 	public UBlueprintAsyncActionBase,
 	public FTickableGameObject
 {
@@ -30,10 +31,10 @@ public:
 	virtual void ReportSuccess() { check(0); }
 	virtual void ReportFailure() { check(0); }
 
-	static void RegisterAsARSystem(const TSharedPtr<FARSystemBase, ESPMode::ThreadSafe>& NewArSystem);
+	static void RegisterAsARSystem(const TSharedPtr<FARSupportInterface , ESPMode::ThreadSafe>& NewArSystem);
 
 protected:
-	static const TSharedPtr<FARSystemBase, ESPMode::ThreadSafe>& GetARSystem();
+	static const TSharedPtr<FARSupportInterface , ESPMode::ThreadSafe>& GetARSystem();
 	/** The async task to check during Tick() */
 	TSharedPtr<FARAsyncTask, ESPMode::ThreadSafe> AsyncTask;
 	
@@ -41,7 +42,7 @@ private:
 	/** True until the async task completes, then false */
 	bool bShouldTick;
 
-	static TSharedPtr<FARSystemBase, ESPMode::ThreadSafe> RegisteredARSystem;
+	static TSharedPtr<FARSupportInterface , ESPMode::ThreadSafe> RegisteredARSystem;
 };
 
 DECLARE_DYNAMIC_MULTICAST_DELEGATE_OneParam(FARSaveWorldPin, const TArray<uint8>&, SavedWorld);
diff --git a/Engine/Source/Runtime/AugmentedReality/Public/ARPin.h b/Engine/Source/Runtime/AugmentedReality/Public/ARPin.h
--- a/Engine/Source/Runtime/AugmentedReality/Public/ARPin.h
+++ b/Engine/Source/Runtime/AugmentedReality/Public/ARPin.h
@@ -5,7 +5,7 @@
 #include "ARTypes.h"
 #include "ARPin.generated.h"
 
-class FARSystemBase;
+class FARSupportInterface ;
 class USceneComponent;
 
 UCLASS(BlueprintType, Experimental, Category="AR AugmentedReality")
@@ -14,7 +14,7 @@ class AUGMENTEDREALITY_API UARPin : public UObject
 	GENERATED_BODY()
 	
 public:
-	virtual void InitARPin( const TSharedRef<FARSystemBase, ESPMode::ThreadSafe>& InTrackingSystemOwner, USceneComponent* InComponentToPin, const FTransform& InLocalToTrackingTransform, UARTrackedGeometry* InTrackedGeometry, const FName InDebugName );
+	virtual void InitARPin( const TSharedRef<FARSupportInterface , ESPMode::ThreadSafe>& InTrackingSystemOwner, USceneComponent* InComponentToPin, const FTransform& InLocalToTrackingTransform, UARTrackedGeometry* InTrackedGeometry, const FName InDebugName );
 
 	/**
 	 * Maps from a Pin's Local Space to the Tracking Space.
@@ -83,7 +83,7 @@ public:
 	
 	
 protected:
-	TSharedPtr<FARSystemBase, ESPMode::ThreadSafe> GetARSystem() const;
+	TSharedPtr<FARSupportInterface , ESPMode::ThreadSafe> GetARSystem() const;
 	
 private:
 	static uint32 DebugPinId;
@@ -103,7 +103,7 @@ private:
 	UPROPERTY()
 	EARTrackingState TrackingState;
 	
-	TWeakPtr<FARSystemBase, ESPMode::ThreadSafe> ARSystem;
+	TWeakPtr<FARSupportInterface , ESPMode::ThreadSafe> ARSystem;
 	
 	FName DebugName;
 	
diff --git a/Engine/Source/Runtime/AugmentedReality/Public/ARSupportInterface.h b/Engine/Source/Runtime/AugmentedReality/Public/ARSupportInterface.h
new file mode 100644
--- /dev/null
+++ b/Engine/Source/Runtime/AugmentedReality/Public/ARSupportInterface.h
@@ -0,0 +1,116 @@
+// Copyright 1998-2018 Epic Games, Inc. All Rights Reserved.
+
+#pragma once
+
+#include "ARTypes.h"
+#include "ARSessionConfig.h"
+#include "ARTextures.h"
+#include "Engine/Engine.h" // for FWorldContext
+
+class IARSystemSupport;
+
+DECLARE_MULTICAST_DELEGATE(FARSystemOnSessionStarted);
+DECLARE_MULTICAST_DELEGATE_OneParam(FARSystemOnAlignmentTransformUpdated, const FTransform&);
+
+/**
+ * Composition Components for tracking system features
+*/
+
+class AUGMENTEDREALITY_API FARSupportInterface  : public FGCObject, public IModularFeature
+{
+public:
+	FARSupportInterface (IARSystemSupport* InARImplementation, IXRTrackingSystem* InXRTrackingSystem);
+	virtual ~FARSupportInterface ();
+
+
+	static FName GetModularFeatureName()
+	{
+		static const FName ModularFeatureName = FName(TEXT("ARSystem"));
+		return ModularFeatureName;
+	}
+
+	void InitializeARSystem();
+	IXRTrackingSystem* GetXRTrackingSystem();
+
+	bool StartARGameFrame(FWorldContext& WorldContext);
+
+	const FTransform& GetAlignmentTransform() const;
+	const UARSessionConfig& GetSessionConfig() const;
+	UARSessionConfig& AccessSessionConfig();
+
+	/** \see UARBlueprintLibrary::GetTrackingQuality() */
+	EARTrackingQuality GetTrackingQuality() const;
+	/** \see UARBlueprintLibrary::StartARSession() */
+	void StartARSession(UARSessionConfig* InSessionConfig);
+	/** \see UARBlueprintLibrary::PauseARSession() */
+	void PauseARSession();
+	/** \see UARBlueprintLibrary::StopARSession() */
+	void StopARSession();
+	/** \see UARBlueprintLibrary::GetARSessionStatus() */
+	FARSessionStatus GetARSessionStatus() const;
+	/** \see UARBlueprintLibrary::IsSessionTypeSupported() */
+	bool IsSessionTypeSupported(EARSessionType SessionType) const;
+
+	/**
+	 * \see UARBlueprintLibrary::SetAlignmentTransform()
+	 * \see IARSystemSupport
+	 * To understand the various spaces involved in Augmented Reality system, \see IARSystemSupport.
+	 */
+	void SetAlignmentTransform(const FTransform& InAlignmentTransform);
+
+	/** \see UARBlueprintLibrary::LineTraceTrackedObjects() */
+	TArray<FARTraceResult> LineTraceTrackedObjects(const FVector2D ScreenCoords, EARLineTraceChannels TraceChannels);
+	/** \see UARBlueprintLibrary::GetAllTrackedGeometries() */
+	TArray<UARTrackedGeometry*> GetAllTrackedGeometries() const;
+	/** \see UARBlueprintLibrary::GetAllPins() */
+	TArray<UARPin*> GetAllPins() const;
+	/** \see UARBlueprintLibrary::GetCameraImage() */
+	UARTextureCameraImage* GetCameraImage();
+	/** \see UARBlueprintLibrary::GetCameraDepth() */
+	UARTextureCameraDepth* GetCameraDepth();
+	/**\see UARBlueprintLibrary::IsEnvironmentCaptureSupported() */
+	bool IsEnvironmentCaptureSupported() const;
+	/**\see UARBlueprintLibrary::AddEnvironmentCaptureProbe() */
+	bool AddManualEnvironmentCaptureProbe(FVector Location, FVector Extent);
+	/** Creates an async task that will perform the work in the background */
+	TSharedPtr<FARGetCandidateObjectAsyncTask, ESPMode::ThreadSafe> GetCandidateObject(FVector Location, FVector Extent) const;
+	/** Creates an async task that will perform the work in the background */
+	TSharedPtr<FARSaveWorldAsyncTask, ESPMode::ThreadSafe> SaveWorld() const;
+	/** @return the current mapping status */
+	EARWorldMappingState GetWorldMappingStatus() const;
+
+	/** \see UARBlueprintLibrary::GetCurrentLightEstimate() */
+	UARLightEstimate* GetCurrentLightEstimate() const;
+
+	/** \see UARBlueprintLibrary::PinComponent() */
+	UARPin* PinComponent(USceneComponent* ComponentToPin, const FTransform& PinToWorldTransform, UARTrackedGeometry* TrackedGeometry = nullptr, const FName DebugName = NAME_None);
+	/** \see UARBlueprintLibrary::PinComponentToTraceResult() */
+	UARPin* PinComponent(USceneComponent* ComponentToPin, const FARTraceResult& HitResult, const FName DebugName = NAME_None);
+	/** \see UARBlueprintLibrary::RemovePin() */
+	void RemovePin(UARPin* PinToRemove);
+
+	/** \see UARBlueprintLibrary::GetSupportedVideoFormats() */
+	TArray<FARVideoFormat> GetSupportedVideoFormats(EARSessionType SessionType = EARSessionType::World) const;
+
+	/** @return the current point cloud data for the ar scene */
+	TArray<FVector> GetPointCloud() const;
+
+	void* GetARSessionRawPointer();
+	void* GetGameThreadARFrameRawPointer();
+
+	//~ FGCObject
+	virtual void AddReferencedObjects(FReferenceCollector& Collector) override;
+	//~ FGCObject
+
+	FARSystemOnSessionStarted OnARSessionStarted;
+	FARSystemOnAlignmentTransformUpdated OnAlignmentTransformUpdated;
+
+private:
+	IARSystemSupport* ARImplemention;
+	IXRTrackingSystem* XRTrackingSystem;
+
+	/** Alignment transform between AR System's tracking space and Unreal's World Space. Useful in static lighting/geometry scenarios. */
+	FTransform AlignmentTransform;
+	UARSessionConfig* ARSettings;
+};
+
diff --git a/Engine/Source/Runtime/AugmentedReality/Public/ARSystem.h b/Engine/Source/Runtime/AugmentedReality/Public/ARSystem.h
--- a/Engine/Source/Runtime/AugmentedReality/Public/ARSystem.h
+++ b/Engine/Source/Runtime/AugmentedReality/Public/ARSystem.h
@@ -2,10 +2,10 @@
 
 #pragma once
 
-#include "Features/IModularFeature.h"
 #include "XRTrackingSystemBase.h"
 #include "ARTypes.h"
 #include "ARSessionConfig.h"
+#include "Engine/Engine.h" // for FWorldContext
 
 class UARPin;
 class UARTrackedGeometry;
@@ -14,12 +14,10 @@ class UARTextureCameraImage;
 class UARTextureCameraDepth;
 struct FARTraceResult;
 
-DECLARE_MULTICAST_DELEGATE(FARSystemOnSessionStarted);
-DECLARE_MULTICAST_DELEGATE_OneParam(FARSystemOnAlignmentTransformUpdated, const FTransform&);
 
 /**
  * Implement IARSystemSupport for any platform that wants to be an Unreal Augmented Reality System. e.g. AppleARKit, GoogleARCore.
- * This interface is included as part of the abstract class \c FARSystemBase. The functions you must override
+ * This interface is included as part of the abstract class \c FARSupportInterface . The functions you must override
  * are coalesced here for clarity.
  * 
  *  Augmented Reality Spaces
@@ -51,6 +49,8 @@ public:
 	/** Invoked after the base AR system has been initialized. */
 	virtual void OnARSystemInitialized(){};
 
+	virtual bool OnStartARGameFrame(FWorldContext& WorldContext) { return false;  };
+
 	/** @return the tracking quality; if unable to determine tracking quality, return EARTrackingQuality::NotAvailable */
 	virtual EARTrackingQuality OnGetTrackingQuality() const = 0;
 	
@@ -134,127 +134,17 @@ public:
 	
 	/** @return The list of supported video formats for this device and session type */
 	virtual TArray<FARVideoFormat> OnGetSupportedVideoFormats(EARSessionType SessionType) const = 0;
-	
+
 	/** @return the current point cloud data for the ar scene */
 	virtual TArray<FVector> OnGetPointCloud() const = 0;
-	
-public:
-	virtual ~IARSystemSupport(){}
-};
-
-class AUGMENTEDREALITY_API FARSystemBase : public IARSystemSupport, public FXRTrackingSystemBase, public FGCObject, public TSharedFromThis<FARSystemBase, ESPMode::ThreadSafe>
-{
-public:
-	//
-	// MODULAR FEATURE SUPPORT
-	//
-	
-	/**
-	 * Part of the pattern for supporting modular features.
-	 *
-	 * @return the name of the feature.
-	 */
-	static FName GetModularFeatureName()
-	{
-		static const FName ModularFeatureName = FName(TEXT("ARSystem"));
-		return ModularFeatureName;
-	}
-
-public:
-	/** Control the construction of AR Systems. \see NewARSystem() */
-	FARSystemBase();
-	void InitializeARSystem();
-	virtual ~FARSystemBase();
-
-public:	
-	/** \see UARBlueprintLibrary::GetTrackingQuality() */
-	EARTrackingQuality GetTrackingQuality() const;
-	/** \see UARBlueprintLibrary::StartARSession() */
-	void StartARSession(UARSessionConfig* InSessionConfig);
-	/** \see UARBlueprintLibrary::PauseARSession() */
-	void PauseARSession();
-	/** \see UARBlueprintLibrary::StopARSession() */
-	void StopARSession();
-	/** \see UARBlueprintLibrary::GetARSessionStatus() */
-	FARSessionStatus GetARSessionStatus() const;
-	/** \see UARBlueprintLibrary::IsSessionTypeSupported() */
-	bool IsSessionTypeSupported(EARSessionType SessionType) const;
-	
-	/**
-	 * \see UARBlueprintLibrary::SetAlignmentTransform()
-	 * \see IARSystemSupport
-	 * To understand the various spaces involved in Augmented Reality system, \see IARSystemSupport.
-	 */
-	void SetAlignmentTransform( const FTransform& InAlignmentTransform );
-	
-	/** \see UARBlueprintLibrary::LineTraceTrackedObjects() */
-	TArray<FARTraceResult> LineTraceTrackedObjects( const FVector2D ScreenCoords, EARLineTraceChannels TraceChannels );
-	/** \see UARBlueprintLibrary::GetAllTrackedGeometries() */
-	TArray<UARTrackedGeometry*> GetAllTrackedGeometries() const;
-	/** \see UARBlueprintLibrary::GetAllPins() */
-	TArray<UARPin*> GetAllPins() const;
-	/** \see UARBlueprintLibrary::GetCameraImage() */
-	UARTextureCameraImage* GetCameraImage();
-	/** \see UARBlueprintLibrary::GetCameraDepth() */
-	UARTextureCameraDepth* GetCameraDepth();
-	/**\see UARBlueprintLibrary::IsEnvironmentCaptureSupported() */
-	bool IsEnvironmentCaptureSupported() const;
-	/**\see UARBlueprintLibrary::AddEnvironmentCaptureProbe() */
-	bool AddManualEnvironmentCaptureProbe(FVector Location, FVector Extent);
-	/** Creates an async task that will perform the work in the background */
-	TSharedPtr<FARGetCandidateObjectAsyncTask, ESPMode::ThreadSafe> GetCandidateObject(FVector Location, FVector Extent) const;
-	/** Creates an async task that will perform the work in the background */
-	TSharedPtr<FARSaveWorldAsyncTask, ESPMode::ThreadSafe> SaveWorld() const;
-	/** @return the current mapping status */
-	EARWorldMappingState GetWorldMappingStatus() const;
-
-	/** \see UARBlueprintLibrary::GetCurrentLightEstimate() */
-	UARLightEstimate* GetCurrentLightEstimate() const;
-	
-	/** \see UARBlueprintLibrary::PinComponent() */
-	UARPin* PinComponent( USceneComponent* ComponentToPin, const FTransform& PinToWorldTransform, UARTrackedGeometry* TrackedGeometry = nullptr, const FName DebugName = NAME_None );
-	/** \see UARBlueprintLibrary::PinComponentToTraceResult() */
-	UARPin* PinComponent( USceneComponent* ComponentToPin, const FARTraceResult& HitResult, const FName DebugName = NAME_None );
-	/** \see UARBlueprintLibrary::RemovePin() */
-	void RemovePin( UARPin* PinToRemove );
-
-	/** \see UARBlueprintLibrary::GetSupportedVideoFormats() */
-	TArray<FARVideoFormat> GetSupportedVideoFormats(EARSessionType SessionType = EARSessionType::World) const;
-	
-	/** \see UARBlueprintLibrary::GetPointCloud() */
-	TArray<FVector> GetPointCloud() const;
 
 	virtual void* GetARSessionRawPointer() = 0;
 	virtual void* GetGameThreadARFrameRawPointer() = 0;
-	
-public:
-	const FTransform& GetAlignmentTransform() const;
-	const UARSessionConfig& GetSessionConfig() const;
-	UARSessionConfig& AccessSessionConfig();
-
-	FARSystemOnSessionStarted OnARSessionStarted;
-	FARSystemOnAlignmentTransformUpdated OnAlignmentTransformUpdated;
-
-protected:
-	void SetAlignmentTransform_Internal(const FTransform& NewAlignmentTransform);
-
-	//~ FGCObject
-	virtual void AddReferencedObjects(FReferenceCollector& Collector) override;
-	//~ FGCObject
-
-private:
-	/** Alignment transform between AR System's tracking space and Unreal's World Space. Useful in static lighting/geometry scenarios. */
-	FTransform AlignmentTransform;
 
-	UARSessionConfig* ARSettings;
+public:
+	virtual ~IARSystemSupport(){}
 };
 
-template<class T, typename... ArgTypes>
-TSharedRef<T, ESPMode::ThreadSafe> NewARSystem( ArgTypes&&... Args )
-{
-	auto NewARSystem = MakeShared<T, ESPMode::ThreadSafe>( Forward<ArgTypes>(Args)... );
-	NewARSystem->InitializeARSystem();
-	return NewARSystem;
-}
+
 
 
diff --git a/Engine/Source/Runtime/AugmentedReality/Public/ARTraceResult.h b/Engine/Source/Runtime/AugmentedReality/Public/ARTraceResult.h
--- a/Engine/Source/Runtime/AugmentedReality/Public/ARTraceResult.h
+++ b/Engine/Source/Runtime/AugmentedReality/Public/ARTraceResult.h
@@ -5,7 +5,7 @@
 #include "ARTypes.h"
 #include "ARTraceResult.generated.h"
 
-class FARSystemBase;
+class FARSupportInterface ;
 
 /**
  * A result of an intersection found during a hit-test.
@@ -17,7 +17,7 @@ struct AUGMENTEDREALITY_API FARTraceResult
 	
 	FARTraceResult();
 	
-	FARTraceResult( const TSharedPtr<FARSystemBase, ESPMode::ThreadSafe>& InARSystem, float InDistanceFromCamera, EARLineTraceChannels InTraceChannel, const FTransform& InLocalToTrackingTransform, UARTrackedGeometry* InTrackedGeometry );
+	FARTraceResult( const TSharedPtr<FARSupportInterface , ESPMode::ThreadSafe>& InARSystem, float InDistanceFromCamera, EARLineTraceChannels InTraceChannel, const FTransform& InLocalToTrackingTransform, UARTrackedGeometry* InTrackedGeometry );
 	
 	float GetDistanceFromCamera() const;
 	
@@ -53,7 +53,7 @@ private:
 	UARTrackedGeometry* TrackedGeometry;
 	
 	/** A reference to the AR system that creates this hit test result. */
-	TSharedPtr<IXRTrackingSystem, ESPMode::ThreadSafe> ARSystem;
+	TSharedPtr<FARSupportInterface , ESPMode::ThreadSafe> ARSystem;
 };
 
 UCLASS()
diff --git a/Engine/Source/Runtime/AugmentedReality/Public/ARTrackable.h b/Engine/Source/Runtime/AugmentedReality/Public/ARTrackable.h
--- a/Engine/Source/Runtime/AugmentedReality/Public/ARTrackable.h
+++ b/Engine/Source/Runtime/AugmentedReality/Public/ARTrackable.h
@@ -5,7 +5,7 @@
 #include "ARTypes.h"
 #include "ARTrackable.generated.h"
 
-class FARSystemBase;
+class FARSupportInterface ;
 class UAREnvironmentCaptureProbeTexture;
 
 UCLASS(BlueprintType)
@@ -20,7 +20,7 @@ public:
 
 	virtual void DebugDraw( UWorld* World, const FLinearColor& OutlineColor, float OutlineThickness, float PersistForSeconds = 0.0f) const;
 	
-	void UpdateTrackedGeometry(const TSharedRef<FARSystemBase, ESPMode::ThreadSafe>& InTrackingSystem, uint32 FrameNumber, double Timestamp, const FTransform& InLocalToTrackingTransform, const FTransform& InAlignmentTransform );
+	void UpdateTrackedGeometry(const TSharedRef<FARSupportInterface , ESPMode::ThreadSafe>& InTrackingSystem, uint32 FrameNumber, double Timestamp, const FTransform& InLocalToTrackingTransform, const FTransform& InAlignmentTransform );
 	
 	void UpdateTrackingState( EARTrackingState NewTrackingState );
 	
@@ -55,7 +55,7 @@ public:
 	float GetLastUpdateTimestamp() const;
 	
 protected:
-	TSharedPtr<FARSystemBase, ESPMode::ThreadSafe> GetARSystem() const;
+	TSharedPtr<FARSupportInterface , ESPMode::ThreadSafe> GetARSystem() const;
 	
 	UPROPERTY()
 	FTransform LocalToTrackingTransform;
@@ -70,7 +70,7 @@ protected:
 	TUniquePtr<IARRef> NativeResource;
 	
 private:
-	TWeakPtr<FARSystemBase, ESPMode::ThreadSafe> ARSystem;
+	TWeakPtr<FARSupportInterface , ESPMode::ThreadSafe> ARSystem;
 	
 	/** The frame number this tracked geometry was last updated on */
 	uint32 LastUpdateFrameNumber;
@@ -89,9 +89,9 @@ class AUGMENTEDREALITY_API UARPlaneGeometry : public UARTrackedGeometry
 	
 public:
 
-	void UpdateTrackedGeometry(const TSharedRef<FARSystemBase, ESPMode::ThreadSafe>& InTrackingSystem, uint32 FrameNumber, double Timestamp, const FTransform& InLocalToTrackingTransform, const FTransform& InAlignmentTransform, const FVector InCenter, const FVector InExtent );
+	void UpdateTrackedGeometry(const TSharedRef<FARSupportInterface , ESPMode::ThreadSafe>& InTrackingSystem, uint32 FrameNumber, double Timestamp, const FTransform& InLocalToTrackingTransform, const FTransform& InAlignmentTransform, const FVector InCenter, const FVector InExtent );
 	
-	void UpdateTrackedGeometry(const TSharedRef<FARSystemBase, ESPMode::ThreadSafe>& InTrackingSystem, uint32 FrameNumber, double Timestamp, const FTransform& InLocalToTrackingTransform, const FTransform& InAlignmentTransform, const FVector InCenter, const FVector InExtent, const TArray<FVector>& InBoundingPoly, UARPlaneGeometry* InSubsumedBy);
+	void UpdateTrackedGeometry(const TSharedRef<FARSupportInterface , ESPMode::ThreadSafe>& InTrackingSystem, uint32 FrameNumber, double Timestamp, const FTransform& InLocalToTrackingTransform, const FTransform& InAlignmentTransform, const FVector InCenter, const FVector InExtent, const TArray<FVector>& InBoundingPoly, UARPlaneGeometry* InSubsumedBy);
 	
 	virtual void DebugDraw( UWorld* World, const FLinearColor& OutlineColor, float OutlineThickness, float PersistForSeconds = 0.0f) const override;
 	
@@ -130,7 +130,7 @@ class AUGMENTEDREALITY_API UARTrackedPoint : public UARTrackedGeometry
 public:
 	virtual void DebugDraw(UWorld* World, const FLinearColor& OutlineColor, float OutlineThickness, float PersistForSeconds = 0.0f) const override;
 
-	void UpdateTrackedGeometry(const TSharedRef<FARSystemBase, ESPMode::ThreadSafe>& InTrackingSystem, uint32 FrameNumber, double Timestamp, const FTransform& InLocalToTrackingTransform, const FTransform& InAlignmentTransform);
+	void UpdateTrackedGeometry(const TSharedRef<FARSupportInterface , ESPMode::ThreadSafe>& InTrackingSystem, uint32 FrameNumber, double Timestamp, const FTransform& InLocalToTrackingTransform, const FTransform& InAlignmentTransform);
 };
 
 UCLASS(BlueprintType)
@@ -141,7 +141,7 @@ class AUGMENTEDREALITY_API UARTrackedImage : public UARTrackedGeometry
 public:
 	virtual void DebugDraw(UWorld* World, const FLinearColor& OutlineColor, float OutlineThickness, float PersistForSeconds = 0.0f) const override;
 
-	void UpdateTrackedGeometry(const TSharedRef<FARSystemBase, ESPMode::ThreadSafe>& InTrackingSystem, uint32 FrameNumber, double Timestamp, const FTransform& InLocalToTrackingTransform, const FTransform& InAlignmentTransform, UARCandidateImage* InDetectedImage);
+	void UpdateTrackedGeometry(const TSharedRef<FARSupportInterface , ESPMode::ThreadSafe>& InTrackingSystem, uint32 FrameNumber, double Timestamp, const FTransform& InLocalToTrackingTransform, const FTransform& InAlignmentTransform, UARCandidateImage* InDetectedImage);
 
 	/** @see DetectedImage */
 	UFUNCTION(BlueprintPure, Category = "AR AugmentedReality|Image Detection")
@@ -258,7 +258,7 @@ class AUGMENTEDREALITY_API UARFaceGeometry : public UARTrackedGeometry
 	GENERATED_BODY()
 	
 public:
-	void UpdateFaceGeometry(const TSharedRef<FARSystemBase, ESPMode::ThreadSafe>& InTrackingSystem, uint32 FrameNumber, double Timestamp, const FTransform& InTransform, const FTransform& InAlignmentTransform, FARBlendShapeMap& InBlendShapes, TArray<FVector>& InVertices, const TArray<int32>& Indices, const FTransform& InLeftEyeTransform, const FTransform& InRightEyeTransform, const FVector& InLookAtTarget);
+	void UpdateFaceGeometry(const TSharedRef<FARSupportInterface , ESPMode::ThreadSafe>& InTrackingSystem, uint32 FrameNumber, double Timestamp, const FTransform& InTransform, const FTransform& InAlignmentTransform, FARBlendShapeMap& InBlendShapes, TArray<FVector>& InVertices, const TArray<int32>& Indices, const FTransform& InLeftEyeTransform, const FTransform& InRightEyeTransform, const FVector& InLookAtTarget);
 	
 	virtual void DebugDraw( UWorld* World, const FLinearColor& OutlineColor, float OutlineThickness, float PersistForSeconds = 0.0f) const override;
 	
@@ -319,7 +319,7 @@ public:
 	/** Draw a box visulizing the bounds of the probe */
 	virtual void DebugDraw(UWorld* World, const FLinearColor& OutlineColor, float OutlineThickness, float PersistForSeconds = 0.0f) const override;
 	
-	void UpdateEnvironmentCapture(const TSharedRef<FARSystemBase, ESPMode::ThreadSafe>& InTrackingSystem, uint32 FrameNumber, double Timestamp, const FTransform& InLocalToTrackingTransform, const FTransform& InAlignmentTransform, FVector InExtent);
+	void UpdateEnvironmentCapture(const TSharedRef<FARSupportInterface , ESPMode::ThreadSafe>& InTrackingSystem, uint32 FrameNumber, double Timestamp, const FTransform& InLocalToTrackingTransform, const FTransform& InAlignmentTransform, FVector InExtent);
 
 	/** @see Extent */
 	UFUNCTION(BlueprintPure, Category="AR AugmentedReality|Environment Capture Probe")
@@ -345,7 +345,7 @@ class AUGMENTEDREALITY_API UARTrackedObject : public UARTrackedGeometry
 public:
 	virtual void DebugDraw(UWorld* World, const FLinearColor& OutlineColor, float OutlineThickness, float PersistForSeconds = 0.0f) const override;
 	
-	void UpdateTrackedGeometry(const TSharedRef<FARSystemBase, ESPMode::ThreadSafe>& InTrackingSystem, uint32 FrameNumber, double Timestamp, const FTransform& InLocalToTrackingTransform, const FTransform& InAlignmentTransform, UARCandidateObject* InDetectedObject);
+	void UpdateTrackedGeometry(const TSharedRef<FARSupportInterface , ESPMode::ThreadSafe>& InTrackingSystem, uint32 FrameNumber, double Timestamp, const FTransform& InLocalToTrackingTransform, const FTransform& InAlignmentTransform, UARCandidateObject* InDetectedObject);
 	
 	/** @see DetectedObject */
 	UFUNCTION(BlueprintPure, Category = "AR AugmentedReality|Object Detection")
diff --git a/Engine/Source/Runtime/AugmentedReality/Public/ARTypes.h b/Engine/Source/Runtime/AugmentedReality/Public/ARTypes.h
--- a/Engine/Source/Runtime/AugmentedReality/Public/ARTypes.h
+++ b/Engine/Source/Runtime/AugmentedReality/Public/ARTypes.h
@@ -6,7 +6,6 @@
 #include "HAL/ThreadSafeBool.h"
 #include "ARTypes.generated.h"
 
-class FARSystemBase;
 class USceneComponent;
 class IXRTrackingSystem;
 class UARPin;
diff --git a/Engine/Source/Runtime/EngineSettings/Classes/GeneralProjectSettings.h b/Engine/Source/Runtime/EngineSettings/Classes/GeneralProjectSettings.h
--- a/Engine/Source/Runtime/EngineSettings/Classes/GeneralProjectSettings.h
+++ b/Engine/Source/Runtime/EngineSettings/Classes/GeneralProjectSettings.h
@@ -79,7 +79,7 @@ class ENGINESETTINGS_API UGeneralProjectSettings
 	UPROPERTY(config, EditAnywhere, Category = Settings)
 	bool bStartInVR;
 	
-	/** This field is no longer used; @see FARSystemBase::StartARSession(); @see UARBlueprintLibrary::StartARSession() */
+	/** This field is no longer used; @see FARSupportInterface ::StartARSession(); @see UARBlueprintLibrary::StartARSession() */
 	UPROPERTY()
 	bool bStartInAR_DEPRECATED;
 
diff --git a/Engine/Source/Runtime/HeadMountedDisplay/HeadMountedDisplay.Build.cs b/Engine/Source/Runtime/HeadMountedDisplay/HeadMountedDisplay.Build.cs
--- a/Engine/Source/Runtime/HeadMountedDisplay/HeadMountedDisplay.Build.cs
+++ b/Engine/Source/Runtime/HeadMountedDisplay/HeadMountedDisplay.Build.cs
@@ -1,4 +1,4 @@
-﻿// Copyright 1998-2018 Epic Games, Inc. All Rights Reserved.
+// Copyright 1998-2018 Epic Games, Inc. All Rights Reserved.
 
 using UnrealBuildTool;
 
@@ -36,7 +36,13 @@ public class HeadMountedDisplay : ModuleRules
             }
         );
 
-		if (Target.bBuildEditor == true)
+        PublicDependencyModuleNames.AddRange(
+            new string[] {
+                "AugmentedReality"
+            }
+        );
+
+        if (Target.bBuildEditor == true)
 		{
 			PrivateDependencyModuleNames.AddRange(
 				new string[]
diff --git a/Engine/Source/Runtime/HeadMountedDisplay/Private/HeadMountedDisplayBase.cpp b/Engine/Source/Runtime/HeadMountedDisplay/Private/HeadMountedDisplayBase.cpp
--- a/Engine/Source/Runtime/HeadMountedDisplay/Private/HeadMountedDisplayBase.cpp
+++ b/Engine/Source/Runtime/HeadMountedDisplay/Private/HeadMountedDisplayBase.cpp
@@ -20,6 +20,14 @@
 // functions (default ctors, etc.) get compiled into this module
 #include "IXRSystemAssets.h"
 
+
+FHeadMountedDisplayBase::FHeadMountedDisplayBase(IARSystemSupport* InARImplementation)
+	: FXRTrackingSystemBase(InARImplementation)
+	, bHeadTrackingEnforced(false)
+{
+}
+
+
 void FHeadMountedDisplayBase::RecordAnalytics()
 {
 	TArray<FAnalyticsEventAttribute> EventAttributes;
diff --git a/Engine/Source/Runtime/HeadMountedDisplay/Private/XRTrackingSystemBase.cpp b/Engine/Source/Runtime/HeadMountedDisplay/Private/XRTrackingSystemBase.cpp
--- a/Engine/Source/Runtime/HeadMountedDisplay/Private/XRTrackingSystemBase.cpp
+++ b/Engine/Source/Runtime/HeadMountedDisplay/Private/XRTrackingSystemBase.cpp
@@ -10,12 +10,20 @@
 // Tracking system delegates
 FXRTrackingSystemDelegates::FXRTrackingOriginChanged FXRTrackingSystemDelegates::OnXRTrackingOriginChanged;
 
-FXRTrackingSystemBase::FXRTrackingSystemBase()
+
+
+
+FXRTrackingSystemBase::FXRTrackingSystemBase(IARSystemSupport* InARImplementation)
 {
+	ARCompositionComponent = MakeShared<FARSupportInterface , ESPMode::ThreadSafe>(InARImplementation, this);
 }
 
 FXRTrackingSystemBase::~FXRTrackingSystemBase()
 {
+	if (ARCompositionComponent.IsValid())
+	{
+		ARCompositionComponent.Reset();
+	}
 }
 
 uint32 FXRTrackingSystemBase::CountTrackedDevices(EXRTrackedDeviceType Type /*= EXRTrackedDeviceType::Any*/)
@@ -103,6 +111,17 @@ void FXRTrackingSystemBase::UpdateExternalTrackingPosition(const FTransform& Ext
 }
 
 
+TSharedPtr<FARSupportInterface , ESPMode::ThreadSafe> FXRTrackingSystemBase::GetARCompositionComponent()
+{
+	return ARCompositionComponent;
+}
+
+const TSharedPtr<const FARSupportInterface , ESPMode::ThreadSafe> FXRTrackingSystemBase::GetARCompositionComponent() const
+{
+	return ARCompositionComponent;
+}
+
+
 FTransform FXRTrackingSystemBase::RefreshTrackingToWorldTransform(FWorldContext& WorldContext)
 {
 	CachedTrackingToWorld = ComputeTrackingToWorldTransform(WorldContext);
diff --git a/Engine/Source/Runtime/HeadMountedDisplay/Public/HeadMountedDisplayBase.h b/Engine/Source/Runtime/HeadMountedDisplay/Public/HeadMountedDisplayBase.h
--- a/Engine/Source/Runtime/HeadMountedDisplay/Public/HeadMountedDisplayBase.h
+++ b/Engine/Source/Runtime/HeadMountedDisplay/Public/HeadMountedDisplayBase.h
@@ -15,10 +15,7 @@ class HEADMOUNTEDDISPLAY_API FHeadMountedDisplayBase : public FXRTrackingSystemB
 {
 
 public:
-	FHeadMountedDisplayBase()
-		: bHeadTrackingEnforced(false)
-	{}
-
+	FHeadMountedDisplayBase(IARSystemSupport* InARImplementation);
 	virtual ~FHeadMountedDisplayBase() {}
 
 	/**
diff --git a/Engine/Source/Runtime/HeadMountedDisplay/Public/IXRTrackingSystem.h b/Engine/Source/Runtime/HeadMountedDisplay/Public/IXRTrackingSystem.h
--- a/Engine/Source/Runtime/HeadMountedDisplay/Public/IXRTrackingSystem.h
+++ b/Engine/Source/Runtime/HeadMountedDisplay/Public/IXRTrackingSystem.h
@@ -162,6 +162,13 @@ public:
 	 */
 	virtual FTransform GetTrackingToWorldTransform() const = 0;
 
+	/**
+	 * This method should return the world to meters scale for the current frame.
+	 * Should be callable on both the render and the game threads.
+	 * @return the current world to meter scale.
+	 */
+	virtual float GetWorldToMetersScale() const = 0;
+
 	/** 
 	 * Computes a transform to convert from 'Floor' origin space to 'Eye' origin space.
 	 * Useful when changing between the two different TrackingOrigin spaces.
diff --git a/Engine/Source/Runtime/HeadMountedDisplay/Public/XRTrackingSystemBase.h b/Engine/Source/Runtime/HeadMountedDisplay/Public/XRTrackingSystemBase.h
--- a/Engine/Source/Runtime/HeadMountedDisplay/Public/XRTrackingSystemBase.h
+++ b/Engine/Source/Runtime/HeadMountedDisplay/Public/XRTrackingSystemBase.h
@@ -4,6 +4,7 @@
 
 #include "IXRTrackingSystem.h"
 #include "IXRCamera.h"
+#include "ARSupportInterface.h"
 
 class HEADMOUNTEDDISPLAY_API FXRTrackingSystemDelegates
 {
@@ -19,7 +20,7 @@ public:
 class HEADMOUNTEDDISPLAY_API FXRTrackingSystemBase : public IXRTrackingSystem
 {
 public:
-	FXRTrackingSystemBase();
+	FXRTrackingSystemBase(IARSystemSupport* InARImplementation);
 	virtual ~FXRTrackingSystemBase();
 
 	/**
@@ -138,12 +139,9 @@ public:
 	*/
 	virtual void UpdateExternalTrackingPosition(const FTransform& ExternalTrackingTransform) override;
 
-	/**
-	 * This method should return the world to meters scale for the current frame.
-	 * Should be callable on both the render and the game threads.
-	 * @return the current world to meter scale.
-	 */
-	virtual float GetWorldToMetersScale() const =0;
+
+	TSharedPtr<FARSupportInterface , ESPMode::ThreadSafe> GetARCompositionComponent();
+	const TSharedPtr<const FARSupportInterface , ESPMode::ThreadSafe> GetARCompositionComponent() const;
 
 protected:
 	/** 
@@ -175,5 +173,8 @@ protected:
 	 * will attempt to normalize the internal tracking system to match this calibration when called.
 	 */
 	FTransform CalibratedOffset;
+
+private:
+	TSharedPtr<FARSupportInterface , ESPMode::ThreadSafe> ARCompositionComponent;
 };
 
