Marked the temp ARKit code from the 4.18 release as deprecated and will delete in 4.23

diff --git a/Engine/Plugins/Runtime/AppleARKit/Source/AppleARKit/Private/AppleARKitAnchor.cpp b/Engine/Plugins/Runtime/AppleARKit/Source/AppleARKit/Private/AppleARKitAnchor.cpp
--- a/Engine/Plugins/Runtime/AppleARKit/Source/AppleARKit/Private/AppleARKitAnchor.cpp
+++ b/Engine/Plugins/Runtime/AppleARKit/Source/AppleARKit/Private/AppleARKitAnchor.cpp
@@ -8,7 +8,7 @@
 // UE4
 #include "Misc/ScopeLock.h"
 
-FTransform UAppleARKitAnchor::GetTransform() const
+FTransform UDEPRECATED_AppleARKitAnchor::GetTransform() const
 {
 	FScopeLock ScopeLock( &UpdateLock );
 
@@ -17,7 +17,7 @@ FTransform UAppleARKitAnchor::GetTransform() const
 
 #if SUPPORTS_ARKIT_1_0
 
-void UAppleARKitAnchor::Update_DelegateThread( ARAnchor* Anchor )
+void UDEPRECATED_AppleARKitAnchor::Update_DelegateThread( ARAnchor* Anchor )
 {
 	FScopeLock ScopeLock( &UpdateLock );
 
diff --git a/Engine/Plugins/Runtime/AppleARKit/Source/AppleARKit/Private/AppleARKitHitTestResult.cpp b/Engine/Plugins/Runtime/AppleARKit/Source/AppleARKit/Private/AppleARKitHitTestResult.cpp
--- a/Engine/Plugins/Runtime/AppleARKit/Source/AppleARKit/Private/AppleARKitHitTestResult.cpp
+++ b/Engine/Plugins/Runtime/AppleARKit/Source/AppleARKit/Private/AppleARKitHitTestResult.cpp
@@ -62,7 +62,7 @@ ARHitTestResultType ToARHitTestResultType(EAppleARKitHitTestResultType InTypes)
 	return Types;
 }
 
-FAppleARKitHitTestResult::FAppleARKitHitTestResult( ARHitTestResult* InARHitTestResult, class UAppleARKitAnchor* InAnchor /*= nullptr*/, float WorldToMetersScale /*= 100.0f*/ )
+FAppleARKitHitTestResult::FAppleARKitHitTestResult( ARHitTestResult* InARHitTestResult, class UDEPRECATED_AppleARKitAnchor* InAnchor /*= nullptr*/, float WorldToMetersScale /*= 100.0f*/ )
 {
 	// Sanity check
 	check( InARHitTestResult );
@@ -71,7 +71,7 @@ FAppleARKitHitTestResult::FAppleARKitHitTestResult( ARHitTestResult* InARHitTest
 	Type = ToEAppleARKitHitTestResultType( InARHitTestResult.type );
     Distance = InARHitTestResult.distance * WorldToMetersScale;
 	Transform = FAppleARKitConversion::ToFTransform( InARHitTestResult.worldTransform );
-	Anchor = InAnchor;
+	Anchor_DEPRECATED = InAnchor;
 }
 
 #endif
diff --git a/Engine/Plugins/Runtime/AppleARKit/Source/AppleARKit/Private/AppleARKitPlaneAnchor.cpp b/Engine/Plugins/Runtime/AppleARKit/Source/AppleARKit/Private/AppleARKitPlaneAnchor.cpp
--- a/Engine/Plugins/Runtime/AppleARKit/Source/AppleARKit/Private/AppleARKitPlaneAnchor.cpp
+++ b/Engine/Plugins/Runtime/AppleARKit/Source/AppleARKit/Private/AppleARKitPlaneAnchor.cpp
@@ -9,21 +9,21 @@
 #include "Misc/ScopeLock.h"
 #include "Math/UnrealMathUtility.h"
 
-FVector UAppleARKitPlaneAnchor::GetCenter() const
+FVector UDEPRECATED_AppleARKitPlaneAnchor::GetCenter() const
 {
 	FScopeLock ScopeLock( &UpdateLock );
 
 	return Center;
 }
 
-FVector UAppleARKitPlaneAnchor::GetExtent() const
+FVector UDEPRECATED_AppleARKitPlaneAnchor::GetExtent() const
 {
 	FScopeLock ScopeLock( &UpdateLock );
 
 	return Extent;
 }
 
-FTransform UAppleARKitPlaneAnchor::GetTransformToCenter() const
+FTransform UDEPRECATED_AppleARKitPlaneAnchor::GetTransformToCenter() const
 {
 	FScopeLock ScopeLock( &UpdateLock );
 
@@ -32,7 +32,7 @@ FTransform UAppleARKitPlaneAnchor::GetTransformToCenter() const
 
 #if SUPPORTS_ARKIT_1_0
 
-void UAppleARKitPlaneAnchor::Update_DelegateThread( ARAnchor* Anchor )
+void UDEPRECATED_AppleARKitPlaneAnchor::Update_DelegateThread( ARAnchor* Anchor )
 {
 	Super::Update_DelegateThread( Anchor );
 
diff --git a/Engine/Plugins/Runtime/AppleARKit/Source/AppleARKit/Public/AppleARKitAnchor.h b/Engine/Plugins/Runtime/AppleARKit/Source/AppleARKit/Public/AppleARKitAnchor.h
--- a/Engine/Plugins/Runtime/AppleARKit/Source/AppleARKit/Public/AppleARKitAnchor.h
+++ b/Engine/Plugins/Runtime/AppleARKit/Source/AppleARKit/Public/AppleARKitAnchor.h
@@ -17,8 +17,8 @@
 // AppleARKit
 #include "AppleARKitAnchor.generated.h"
 
-UCLASS( BlueprintType )
-class APPLEARKIT_API UAppleARKitAnchor : public UObject
+UCLASS( BlueprintType, Deprecated )
+class APPLEARKIT_API UDEPRECATED_AppleARKitAnchor : public UObject
 {
 	GENERATED_BODY()
 
diff --git a/Engine/Plugins/Runtime/AppleARKit/Source/AppleARKit/Public/AppleARKitBlueprintLibrary.h b/Engine/Plugins/Runtime/AppleARKit/Source/AppleARKit/Public/AppleARKitBlueprintLibrary.h
--- a/Engine/Plugins/Runtime/AppleARKit/Source/AppleARKit/Public/AppleARKitBlueprintLibrary.h
+++ b/Engine/Plugins/Runtime/AppleARKit/Source/AppleARKit/Public/AppleARKitBlueprintLibrary.h
@@ -23,7 +23,7 @@ public:
 	 * of the game thread update by watching for changes to GFrameNumber to trigger pulling a
 	 * new frame from the session.
 	 */
-	UFUNCTION( BlueprintCallable, Category="AppleARKit", meta = (WorldContext = "WorldContextObject", Keywords = "ar augmentedreality augmented reality frame"))
+	UFUNCTION( BlueprintCallable, Category="AppleARKit", meta = (DeprecatedFunction, WorldContext = "WorldContextObject", Keywords = "ar augmentedreality augmented reality frame"))
 	static bool GetCurrentFrame( UObject* WorldContextObject, FAppleARKitFrame& OutCurrentFrame );
 	
 	/**
@@ -37,6 +37,6 @@ public:
 	 *
 	 * @param ScreenPosition The viewport pixel coordinate of the trace origin.
 	 */
-	UFUNCTION( BlueprintCallable, Category="AppleARKit", meta = (WorldContext = "WorldContextObject", Keywords = "ar augmentedreality augmented reality trace hittest hit line"))
+	UFUNCTION( BlueprintCallable, Category="AppleARKit", meta = (DeprecatedFunction, WorldContext = "WorldContextObject", Keywords = "ar augmentedreality augmented reality trace hittest hit line"))
 	static bool HitTestAtScreenPosition_TrackingSpace( UObject* WorldContextObject, const FVector2D ScreenPosition, EAppleARKitHitTestResultType Types, TArray< FAppleARKitHitTestResult >& OutResults );
 };
diff --git a/Engine/Plugins/Runtime/AppleARKit/Source/AppleARKit/Public/AppleARKitHitTestResult.h b/Engine/Plugins/Runtime/AppleARKit/Source/AppleARKit/Public/AppleARKitHitTestResult.h
--- a/Engine/Plugins/Runtime/AppleARKit/Source/AppleARKit/Public/AppleARKitHitTestResult.h
+++ b/Engine/Plugins/Runtime/AppleARKit/Source/AppleARKit/Public/AppleARKitHitTestResult.h
@@ -62,7 +62,7 @@ struct APPLEARKIT_API FAppleARKitHitTestResult
 	 * This is a conversion copy-constructor that takes a raw ARHitTestResult and fills this 
 	 * structs members with the UE4-ified versions of ARHitTestResult's properties.
 	 */ 
-	FAppleARKitHitTestResult( ARHitTestResult* InARHitTestResult, class UAppleARKitAnchor* InAnchor = nullptr, float WorldToMetersScale = 100.0f );
+	FAppleARKitHitTestResult( ARHitTestResult* InARHitTestResult, class UDEPRECATED_AppleARKitAnchor* InAnchor = nullptr, float WorldToMetersScale = 100.0f );
 
 #endif
 
@@ -90,6 +90,6 @@ struct APPLEARKIT_API FAppleARKitHitTestResult
 	 * 
 	 * An anchor will only be provided for existing plane result types.
 	 */
-	UPROPERTY( BlueprintReadOnly, Category = "AppleARKitHitTestResult")
-	class UAppleARKitAnchor* Anchor = nullptr;
+	UPROPERTY()
+	class UDEPRECATED_AppleARKitAnchor* Anchor_DEPRECATED = nullptr;
 };
diff --git a/Engine/Plugins/Runtime/AppleARKit/Source/AppleARKit/Public/AppleARKitPlaneAnchor.h b/Engine/Plugins/Runtime/AppleARKit/Source/AppleARKit/Public/AppleARKitPlaneAnchor.h
--- a/Engine/Plugins/Runtime/AppleARKit/Source/AppleARKit/Public/AppleARKitPlaneAnchor.h
+++ b/Engine/Plugins/Runtime/AppleARKit/Source/AppleARKit/Public/AppleARKitPlaneAnchor.h
@@ -14,8 +14,8 @@
 #include "AppleARKitAnchor.h"
 #include "AppleARKitPlaneAnchor.generated.h"
 
-UCLASS( BlueprintType )
-class APPLEARKIT_API UAppleARKitPlaneAnchor : public UAppleARKitAnchor
+UCLASS( BlueprintType, Deprecated )
+class APPLEARKIT_API UDEPRECATED_AppleARKitPlaneAnchor : public UDEPRECATED_AppleARKitAnchor
 {
 	GENERATED_BODY()
 
