Grouped Apple/Google AR plugins under an AR directory in Experimental to improve the organization of various AR plugins

diff --git a/Engine/Plugins/Experimental/AR/Apple/AppleARKit/Source/AppleARKit/AppleARKit.Build.cs b/Engine/Plugins/Experimental/AR/Apple/AppleARKit/Source/AppleARKit/AppleARKit.Build.cs
new file mode 100644
--- /dev/null
+++ b/Engine/Plugins/Experimental/AR/Apple/AppleARKit/Source/AppleARKit/AppleARKit.Build.cs
@@ -0,0 +1,63 @@
+// Copyright 1998-2018 Epic Games, Inc. All Rights Reserved.
+
+using UnrealBuildTool;
+using System.IO;
+
+public class AppleARKit : ModuleRules
+{
+	public AppleARKit(ReadOnlyTargetRules Target) : base(Target)
+	{
+		PrivateIncludePaths.AddRange(
+			new string[] {
+				"../../../../Source/Runtime/Renderer/Private",
+				// ... add other private include paths required here ...
+			}
+			);
+			
+		
+		PublicDependencyModuleNames.AddRange(
+			new string[]
+			{
+				"Core",
+				"Engine",
+				// ... add other public dependencies that you statically link with here ...
+			}
+			);
+			
+		
+		PrivateDependencyModuleNames.AddRange(
+			new string[]
+			{
+				"CoreUObject",
+				"Slate",
+				"SlateCore",
+				"RHI",
+                "Renderer",
+                "RenderCore",
+                "ShaderCore",
+                "HeadMountedDisplay",
+                "AugmentedReality",
+                "AppleImageUtils"
+				// ... add private dependencies that you statically link with here ...
+			}
+			);
+		
+		
+		DynamicallyLoadedModuleNames.AddRange(
+			new string[]
+			{
+				// ... add any modules that your module loads dynamically here ...
+			}
+			);
+
+		if (Target.Platform == UnrealTargetPlatform.IOS)
+		{
+    		PrivateDependencyModuleNames.Add("IOSRuntimeSettings");
+
+			PublicFrameworks.Add( "ARKit" );
+
+            string PluginPath = Utils.MakePathRelativeTo(ModuleDirectory, Target.RelativeEnginePath);
+            AdditionalPropertiesForReceipt.Add("IOSPlugin", Path.Combine(PluginPath, "AppleARKit_IOS_UPL.xml"));
+		}
+	}
+}
diff --git a/Engine/Plugins/Experimental/AR/Apple/AppleARKit/Source/AppleARKit/Private/AppleARKitAnchor.cpp b/Engine/Plugins/Experimental/AR/Apple/AppleARKit/Source/AppleARKit/Private/AppleARKitAnchor.cpp
new file mode 100644
--- /dev/null
+++ b/Engine/Plugins/Experimental/AR/Apple/AppleARKit/Source/AppleARKit/Private/AppleARKitAnchor.cpp
@@ -0,0 +1,28 @@
+// Copyright 1998-2018 Epic Games, Inc. All Rights Reserved.
+
+// AppleARKit
+#include "AppleARKitAnchor.h"
+#include "AppleARKitModule.h"
+#include "AppleARKitConversion.h"
+
+// UE4
+#include "Misc/ScopeLock.h"
+
+FTransform UDEPRECATED_AppleARKitAnchor::GetTransform() const
+{
+	FScopeLock ScopeLock( &UpdateLock );
+
+	return Transform;
+}
+
+#if SUPPORTS_ARKIT_1_0
+
+void UDEPRECATED_AppleARKitAnchor::Update_DelegateThread( ARAnchor* Anchor )
+{
+	FScopeLock ScopeLock( &UpdateLock );
+
+	// @todo arkit use World Settings WorldToMetersScale
+	Transform = FAppleARKitConversion::ToFTransform( Anchor.transform );
+}
+
+#endif
diff --git a/Engine/Plugins/Experimental/AR/Apple/AppleARKit/Source/AppleARKit/Private/AppleARKitBlueprintLibrary.cpp b/Engine/Plugins/Experimental/AR/Apple/AppleARKit/Source/AppleARKit/Private/AppleARKitBlueprintLibrary.cpp
new file mode 100644
--- /dev/null
+++ b/Engine/Plugins/Experimental/AR/Apple/AppleARKit/Source/AppleARKit/Private/AppleARKitBlueprintLibrary.cpp
@@ -0,0 +1,26 @@
+// Copyright 1998-2018 Epic Games, Inc. All Rights Reserved.
+
+#include "AppleARKitBlueprintLibrary.h"
+#include "AppleARKitModule.h"
+
+
+bool UAppleARKitBlueprintLibrary::GetCurrentFrame( UObject* WorldContextObject, FAppleARKitFrame& OutCurrentFrame )
+{
+	auto System = FAppleARKitModule::GetARKitSystem();
+	if (System.IsValid())
+	{
+		return System->GetCurrentFrame(OutCurrentFrame);
+	}
+	return false;
+}
+	
+bool UAppleARKitBlueprintLibrary::HitTestAtScreenPosition_TrackingSpace( UObject* WorldContextObject, const FVector2D ScreenPosition, EAppleARKitHitTestResultType Types, TArray< FAppleARKitHitTestResult >& OutResults )
+{
+	auto System = FAppleARKitModule::GetARKitSystem();
+	if (System.IsValid())
+	{
+		return System->HitTestAtScreenPosition(ScreenPosition, Types, OutResults);
+	}
+	return false;
+}
+
diff --git a/Engine/Plugins/Experimental/AR/Apple/AppleARKit/Source/AppleARKit/Private/AppleARKitFrame.cpp b/Engine/Plugins/Experimental/AR/Apple/AppleARKit/Source/AppleARKit/Private/AppleARKitFrame.cpp
new file mode 100644
--- /dev/null
+++ b/Engine/Plugins/Experimental/AR/Apple/AppleARKit/Source/AppleARKit/Private/AppleARKitFrame.cpp
@@ -0,0 +1,213 @@
+// Copyright 1998-2018 Epic Games, Inc. All Rights Reserved.
+
+// AppleARKit
+#include "AppleARKitFrame.h"
+#include "AppleARKitModule.h"
+#include "Misc/ScopeLock.h"
+
+// Default constructor
+FAppleARKitFrame::FAppleARKitFrame()
+	: Timestamp(0.0)
+#if SUPPORTS_ARKIT_1_0
+	, CapturedYImage(nullptr)
+	, CapturedCbCrImage(nullptr)
+	, CameraImage(nullptr)
+	, CameraDepth(nullptr)
+	, NativeFrame(nullptr)
+#endif
+	, CapturedYImageWidth(0)
+	, CapturedYImageHeight(0)
+	, CapturedCbCrImageWidth(0)
+	, CapturedCbCrImageHeight(0)
+	, WorldMappingState(EARWorldMappingState::NotAvailable)
+{
+};
+
+#if SUPPORTS_ARKIT_2_0
+EARWorldMappingState ToEARWorldMappingState(ARWorldMappingStatus MapStatus)
+{
+	switch (MapStatus)
+	{
+		// These both mean more data is needed
+		case ARWorldMappingStatusLimited:
+		case ARWorldMappingStatusExtending:
+			return EARWorldMappingState::StillMappingNotRelocalizable;
+
+		case ARWorldMappingStatusMapped:
+			return EARWorldMappingState::Mapped;
+	}
+	return EARWorldMappingState::NotAvailable;
+}
+#endif
+
+#if SUPPORTS_ARKIT_1_0
+
+FAppleARKitFrame::FAppleARKitFrame( ARFrame* InARFrame, CVMetalTextureCacheRef MetalTextureCache )
+	: Camera( InARFrame.camera )
+	, LightEstimate( InARFrame.lightEstimate )
+	, WorldMappingState(EARWorldMappingState::NotAvailable)
+{
+	// Sanity check
+	check( InARFrame );
+	check( MetalTextureCache );
+
+	// Copy timestamp
+	Timestamp = InARFrame.timestamp;
+
+	CameraImage = nullptr;
+	CameraDepth = nullptr;
+
+	// Copy / convert pass-through camera image's CVPixelBuffer to an MTLTexture so we can pass it
+	// directly to FTextureResource's.
+	// @see AppleARKitCameraTextureResource.cpp
+	CapturedYImage = nullptr;
+	CapturedYImageWidth = InARFrame.camera.imageResolution.width;
+	CapturedYImageHeight = InARFrame.camera.imageResolution.height;
+	
+	CapturedCbCrImage = nullptr;
+	CapturedCbCrImageWidth = InARFrame.camera.imageResolution.width;
+	CapturedCbCrImageHeight = InARFrame.camera.imageResolution.height;
+	
+	if ( InARFrame.capturedImage )
+	{
+		CameraImage = InARFrame.capturedImage;
+		CFRetain(CameraImage);
+
+		// Update SizeX & Y
+		CapturedYImageWidth = CVPixelBufferGetWidthOfPlane( InARFrame.capturedImage, 0 );
+		CapturedYImageHeight = CVPixelBufferGetHeightOfPlane( InARFrame.capturedImage, 0 );
+		CapturedCbCrImageWidth = CVPixelBufferGetWidthOfPlane( InARFrame.capturedImage, 1 );
+		CapturedCbCrImageHeight = CVPixelBufferGetHeightOfPlane( InARFrame.capturedImage, 1 );
+		
+		// Create a metal texture from the CVPixelBufferRef. The CVMetalTextureRef will
+		// be released in the FAppleARKitFrame destructor.
+		// NOTE: On success, CapturedImage will be a new CVMetalTextureRef with a ref count of 1
+		// 		 so we don't need to CFRetain it. The corresponding CFRelease is handled in
+		//
+		CVReturn Result = CVMetalTextureCacheCreateTextureFromImage( nullptr, MetalTextureCache, InARFrame.capturedImage, nullptr, MTLPixelFormatR8Unorm, CapturedYImageWidth, CapturedYImageHeight, /*PlaneIndex*/0, &CapturedYImage );
+		check( Result == kCVReturnSuccess );
+		check( CapturedYImage );
+		check( CFGetRetainCount(CapturedYImage) == 1);
+		
+		Result = CVMetalTextureCacheCreateTextureFromImage( nullptr, MetalTextureCache, InARFrame.capturedImage, nullptr, MTLPixelFormatRG8Unorm, CapturedCbCrImageWidth, CapturedCbCrImageHeight, /*PlaneIndex*/1, &CapturedCbCrImage );
+		check( Result == kCVReturnSuccess );
+		check( CapturedCbCrImage );
+		check( CFGetRetainCount(CapturedCbCrImage) == 1);
+	}
+
+	if (InARFrame.capturedDepthData)
+	{
+		CameraDepth = InARFrame.capturedDepthData;
+		CFRetain(CameraDepth);
+	}
+
+	NativeFrame = (void*)CFRetain(InARFrame);
+
+#if SUPPORTS_ARKIT_2_0
+	if (FAppleARKitAvailability::SupportsARKit20())
+	{
+		WorldMappingState = ToEARWorldMappingState(InARFrame.worldMappingStatus);
+	}
+#endif
+}
+
+FAppleARKitFrame::FAppleARKitFrame( const FAppleARKitFrame& Other )
+	: Timestamp( Other.Timestamp )
+	, CapturedYImage( nullptr )
+	, CapturedCbCrImage( nullptr )
+	, CameraImage( nullptr )
+	, CameraDepth( nullptr )
+	, CapturedYImageWidth( Other.CapturedYImageWidth )
+	, CapturedYImageHeight( Other.CapturedYImageHeight )
+	, CapturedCbCrImageWidth( Other.CapturedCbCrImageWidth )
+	, CapturedCbCrImageHeight( Other.CapturedCbCrImageHeight )
+	, Camera( Other.Camera )
+	, LightEstimate( Other.LightEstimate )
+	, WorldMappingState(Other.WorldMappingState)
+{
+	if(Other.NativeFrame != nullptr)
+	{
+		NativeFrame = (void*)CFRetain((CFTypeRef)Other.NativeFrame);
+	}
+}
+
+FAppleARKitFrame::~FAppleARKitFrame()
+{
+	// Release captured image
+	if ( CapturedYImage != nullptr )
+	{
+		CFRelease( CapturedYImage );
+	}
+	if ( CapturedCbCrImage != nullptr )
+	{
+		CFRelease( CapturedCbCrImage );
+	}
+	if (CameraImage != nullptr)
+	{
+		CFRelease(CameraImage);
+	}
+	if (CameraDepth != nullptr)
+	{
+		CFRelease(CameraDepth);
+	}
+	if(NativeFrame != nullptr)
+	{
+		CFRelease((CFTypeRef)NativeFrame);
+	}
+}
+
+FAppleARKitFrame& FAppleARKitFrame::operator=( const FAppleARKitFrame& Other )
+{
+	if (&Other == this)
+	{
+		return *this;
+	}
+
+	// Release outgoing image
+	if ( CapturedYImage != nullptr )
+	{
+		CFRelease( CapturedYImage );
+	}
+	if ( CapturedCbCrImage != nullptr )
+	{
+		CFRelease( CapturedCbCrImage );
+	}
+	if (CameraImage != nullptr)
+	{
+		CFRelease(CameraImage);
+		CameraImage = nullptr;
+	}
+	if (CameraDepth != nullptr)
+	{
+		CFRelease(CameraDepth);
+		CameraDepth = nullptr;
+	}
+	if(NativeFrame != nullptr)
+	{
+		CFRelease((CFTypeRef)NativeFrame);
+		NativeFrame = nullptr;
+	}
+
+	if(Other.NativeFrame != nullptr)
+	{
+		NativeFrame = (void*)CFRetain((CFTypeRef)Other.NativeFrame);
+	}
+	
+	// Member-wise copy
+	Timestamp = Other.Timestamp;
+	CapturedYImage = nullptr;
+	CapturedYImageWidth = Other.CapturedYImageWidth;
+	CapturedYImageHeight = Other.CapturedYImageHeight;
+	CapturedCbCrImage = nullptr;
+	CapturedCbCrImageWidth = Other.CapturedCbCrImageWidth;
+	CapturedCbCrImageHeight = Other.CapturedCbCrImageHeight;
+	Camera = Other.Camera;
+	LightEstimate = Other.LightEstimate;
+	WorldMappingState = Other.WorldMappingState;
+
+	NativeFrame = Other.NativeFrame;
+
+	return *this;
+}
+
+#endif
diff --git a/Engine/Plugins/Experimental/AR/Apple/AppleARKit/Source/AppleARKit/Private/AppleARKitHitTestResult.cpp b/Engine/Plugins/Experimental/AR/Apple/AppleARKit/Source/AppleARKit/Private/AppleARKitHitTestResult.cpp
new file mode 100644
--- /dev/null
+++ b/Engine/Plugins/Experimental/AR/Apple/AppleARKit/Source/AppleARKit/Private/AppleARKitHitTestResult.cpp
@@ -0,0 +1,77 @@
+// Copyright 1998-2018 Epic Games, Inc. All Rights Reserved.
+
+// AppleARKit
+#include "AppleARKitHitTestResult.h"
+#include "AppleARKitConversion.h"
+
+#if SUPPORTS_ARKIT_1_0
+
+/** Conversion function from ARKit native ARHitTestResultType */
+EAppleARKitHitTestResultType ToEAppleARKitHitTestResultType(ARHitTestResultType InTypes)
+{
+	EAppleARKitHitTestResultType Types = EAppleARKitHitTestResultType::None;
+
+	if (!!(InTypes & ARHitTestResultTypeFeaturePoint))
+	{
+		Types |= EAppleARKitHitTestResultType::FeaturePoint;
+	}
+
+	if (!!(InTypes & ARHitTestResultTypeEstimatedHorizontalPlane))
+	{
+		Types |= EAppleARKitHitTestResultType::EstimatedHorizontalPlane;
+	}
+
+	if (!!(InTypes & ARHitTestResultTypeExistingPlane))
+	{
+		Types |= EAppleARKitHitTestResultType::ExistingPlane;
+	}
+
+    if (!!(InTypes & ARHitTestResultTypeExistingPlaneUsingExtent))
+    {
+        Types |= EAppleARKitHitTestResultType::ExistingPlaneUsingExtent;
+    }
+    
+	return Types;
+}
+
+/** Conversion function to ARKit native ARHitTestResultType */
+ARHitTestResultType ToARHitTestResultType(EAppleARKitHitTestResultType InTypes)
+{
+	ARHitTestResultType Types = 0;
+
+	if (!!(InTypes & EAppleARKitHitTestResultType::FeaturePoint))
+	{
+		Types |= ARHitTestResultTypeFeaturePoint;
+	}
+
+	if (!!(InTypes & EAppleARKitHitTestResultType::EstimatedHorizontalPlane))
+	{
+		Types |= ARHitTestResultTypeEstimatedHorizontalPlane;
+	}
+
+	if (!!(InTypes & EAppleARKitHitTestResultType::ExistingPlane))
+	{
+		Types |= ARHitTestResultTypeExistingPlane;
+	}
+
+    if (!!(InTypes & EAppleARKitHitTestResultType::ExistingPlaneUsingExtent))
+    {
+        Types |= ARHitTestResultTypeExistingPlaneUsingExtent;
+    }
+    
+	return Types;
+}
+
+FAppleARKitHitTestResult::FAppleARKitHitTestResult( ARHitTestResult* InARHitTestResult, class UDEPRECATED_AppleARKitAnchor* InAnchor /*= nullptr*/, float WorldToMetersScale /*= 100.0f*/ )
+{
+	// Sanity check
+	check( InARHitTestResult );
+
+	// Convert properties
+	Type = ToEAppleARKitHitTestResultType( InARHitTestResult.type );
+    Distance = InARHitTestResult.distance * WorldToMetersScale;
+	Transform = FAppleARKitConversion::ToFTransform( InARHitTestResult.worldTransform );
+	Anchor_DEPRECATED = InAnchor;
+}
+
+#endif
diff --git a/Engine/Plugins/Experimental/AR/Apple/AppleARKit/Source/AppleARKit/Private/AppleARKitPlaneAnchor.cpp b/Engine/Plugins/Experimental/AR/Apple/AppleARKit/Source/AppleARKit/Private/AppleARKitPlaneAnchor.cpp
new file mode 100644
--- /dev/null
+++ b/Engine/Plugins/Experimental/AR/Apple/AppleARKit/Source/AppleARKit/Private/AppleARKitPlaneAnchor.cpp
@@ -0,0 +1,52 @@
+// Copyright 1998-2018 Epic Games, Inc. All Rights Reserved.
+
+// AppleARKit
+#include "AppleARKitPlaneAnchor.h"
+#include "AppleARKitModule.h"
+#include "AppleARKitConversion.h"
+
+// UE4
+#include "Misc/ScopeLock.h"
+#include "Math/UnrealMathUtility.h"
+
+FVector UDEPRECATED_AppleARKitPlaneAnchor::GetCenter() const
+{
+	FScopeLock ScopeLock( &UpdateLock );
+
+	return Center;
+}
+
+FVector UDEPRECATED_AppleARKitPlaneAnchor::GetExtent() const
+{
+	FScopeLock ScopeLock( &UpdateLock );
+
+	return Extent;
+}
+
+FTransform UDEPRECATED_AppleARKitPlaneAnchor::GetTransformToCenter() const
+{
+	FScopeLock ScopeLock( &UpdateLock );
+
+	return FTransform( Center ) * Transform;
+}
+
+#if SUPPORTS_ARKIT_1_0
+
+void UDEPRECATED_AppleARKitPlaneAnchor::Update_DelegateThread( ARAnchor* Anchor )
+{
+	Super::Update_DelegateThread( Anchor );
+
+	// Plane anchor?
+	if ([Anchor isKindOfClass:[ARPlaneAnchor class]])
+	{
+		ARPlaneAnchor* PlaneAnchor = (ARPlaneAnchor*)Anchor;
+
+		FScopeLock ScopeLock( &UpdateLock );
+		
+		// @todo use World Settings WorldToMetersScale
+		Extent = FAppleARKitConversion::ToFVector( PlaneAnchor.extent ).GetAbs();
+		Center = FAppleARKitConversion::ToFVector( PlaneAnchor.center );
+	}
+}
+
+#endif
diff --git a/Engine/Plugins/Experimental/AR/Apple/AppleARKit/Source/AppleARKit/Private/AppleARKitVideoOverlay.cpp b/Engine/Plugins/Experimental/AR/Apple/AppleARKit/Source/AppleARKit/Private/AppleARKitVideoOverlay.cpp
new file mode 100644
--- /dev/null
+++ b/Engine/Plugins/Experimental/AR/Apple/AppleARKit/Source/AppleARKit/Private/AppleARKitVideoOverlay.cpp
@@ -0,0 +1,482 @@
+// Copyright 1998-2018 Epic Games, Inc. All Rights Reserved.
+
+#include "AppleARKitVideoOverlay.h"
+#include "AppleARKitFrame.h"
+#include "Materials/Material.h"
+#include "MaterialShaderType.h"
+#include "MaterialShader.h"
+#include "ExternalTexture.h"
+#include "Misc/Guid.h"
+#include "ExternalTextureGuid.h"
+#include "Containers/ResourceArray.h"
+#include "MediaShaders.h"
+#include "PipelineStateCache.h"
+#include "RHIUtilities.h"
+#include "RHIStaticStates.h"
+#include "EngineModule.h"
+#include "SceneUtils.h"
+#include "RendererInterface.h"
+#include "ScreenRendering.h"
+#include "Containers/DynamicRHIResourceArray.h"
+#include "PostProcess/SceneFilterRendering.h"
+#include "PostProcessParameters.h"
+#if SUPPORTS_ARKIT_1_0
+	#include "IOSAppDelegate.h"
+#endif
+
+#if SUPPORTS_ARKIT_1_0
+/**
+* Passes a CVMetalTextureRef through to the RHI to wrap in an RHI texture without traversing system memory.
+* @see FAvfTexture2DResourceWrapper & FMetalSurface::FMetalSurface
+*/
+class FAppleARKitCameraTextureResourceWrapper : public FResourceBulkDataInterface
+{
+public:
+	FAppleARKitCameraTextureResourceWrapper(CFTypeRef InImageBuffer)
+		: ImageBuffer(InImageBuffer)
+	{
+		check(ImageBuffer);
+		CFRetain(ImageBuffer);
+	}
+
+	/**
+	* @return ptr to the resource memory which has been preallocated
+	*/
+	virtual const void* GetResourceBulkData() const override
+	{
+		return ImageBuffer;
+	}
+
+	/**
+	* @return size of resource memory
+	*/
+	virtual uint32 GetResourceBulkDataSize() const override
+	{
+		return 0;
+	}
+
+	/**
+	* @return the type of bulk data for special handling
+	*/
+	virtual EBulkDataType GetResourceType() const override
+	{
+		return EBulkDataType::MediaTexture;
+	}
+
+	/**
+	* Free memory after it has been used to initialize RHI resource
+	*/
+	virtual void Discard() override
+	{
+		delete this;
+	}
+
+	virtual ~FAppleARKitCameraTextureResourceWrapper()
+	{
+		CFRelease(ImageBuffer);
+		ImageBuffer = nullptr;
+	}
+
+	CFTypeRef ImageBuffer;
+};
+#endif
+
+FAppleARKitVideoOverlay::FAppleARKitVideoOverlay()
+	: RenderingOverlayMaterial(nullptr)
+	, LastUpdateTimestamp(-1.0)
+{
+	RenderingOverlayMaterial = LoadObject<UMaterialInterface>(nullptr, TEXT("/AppleARKit/ARKitCameraMaterial.ARKitCameraMaterial"));
+	check(RenderingOverlayMaterial != nullptr);
+	RenderingOverlayMaterial->AddToRoot();
+}
+
+void FAppleARKitVideoOverlay::UpdateVideoTexture_RenderThread(FRHICommandListImmediate& RHICmdList, FAppleARKitFrame& Frame, const FSceneViewFamily& InViewFamily)
+{
+	// Allocate and register
+	if (VideoTextureY == nullptr)
+	{
+		check(VideoTextureCbCr == nullptr);
+		check(OverlayIndexBufferRHI == nullptr);
+		check(OverlayVertexBufferRHI[0] == nullptr);
+
+		FRHIResourceCreateInfo CreateInfo;
+		VideoTextureY = RHICmdList.CreateTexture2D(1, 1, PF_R8G8B8A8, 1, 1, 0, CreateInfo);
+		VideoTextureCbCr = RHICmdList.CreateTexture2D(1, 1, PF_R8G8B8A8, 1, 1, 0, CreateInfo);
+
+		FSamplerStateInitializerRHI SamplerStateInitializer(SF_Bilinear, AM_Wrap, AM_Wrap, AM_Wrap);
+		FSamplerStateRHIRef SamplerStateRHI = RHICreateSamplerState(SamplerStateInitializer);
+
+		FExternalTextureRegistry::Get().RegisterExternalTexture(ARKitPassthroughCameraExternalTextureYGuid, VideoTextureY, SamplerStateRHI);
+		FExternalTextureRegistry::Get().RegisterExternalTexture(ARKitPassthroughCameraExternalTextureCbCrGuid, VideoTextureCbCr, SamplerStateRHI);
+
+		// Setup index buffer
+		const uint16 Indices[] = { 0, 1, 2, 2, 1, 3 };
+
+		TResourceArray<uint16, INDEXBUFFER_ALIGNMENT> IndexBuffer;
+		const uint32 NumIndices = ARRAY_COUNT(Indices);
+		IndexBuffer.AddUninitialized(NumIndices);
+		FMemory::Memcpy(IndexBuffer.GetData(), Indices, NumIndices * sizeof(uint16));
+
+		FRHIResourceCreateInfo CreateInfoIB(&IndexBuffer);
+		OverlayIndexBufferRHI = RHICreateIndexBuffer(sizeof(uint16), IndexBuffer.GetResourceDataSize(), BUF_Static, CreateInfoIB);
+
+		check(InViewFamily.Views.Num() > 0);
+		const FSceneView& View = *InViewFamily.Views[0];
+
+		const FVector2D ViewSize(View.UnconstrainedViewRect.Max.X, View.UnconstrainedViewRect.Max.Y);
+		
+		// CameraSize is 1280 x 720 regardless of the device orientation. We flip it here if needed to make it consistent with the view size.
+		FVector2D CameraSize = Frame.Camera.ImageResolution;
+		if ((ViewSize.X > ViewSize.Y) != (CameraSize.X > CameraSize.Y))
+		{
+			CameraSize = FVector2D(CameraSize.Y, CameraSize.X);
+		}
+		
+		const float CameraAspectRatio = CameraSize.X / CameraSize.Y;
+		const float ViewAspectRatio = ViewSize.X / ViewSize.Y;
+		const float ViewAspectRatioLandscape = (ViewSize.X > ViewSize.Y) ? ViewAspectRatio : ViewSize.Y / ViewSize.X;
+		
+		float UVOffsetAmount = 0.0f;
+		if (!FMath::IsNearlyEqual(ViewAspectRatio, CameraAspectRatio))
+		{
+			if (ViewAspectRatio > CameraAspectRatio)
+			{
+				UVOffsetAmount = 0.5f * (1.0f - (CameraAspectRatio / ViewAspectRatio));
+			}
+			else
+			{
+				UVOffsetAmount = 0.5f * (1.0f - (ViewAspectRatio / CameraAspectRatio));
+			}
+		}
+		
+		const FVector2D UVOffset = (ViewAspectRatioLandscape <= Frame.Camera.GetAspectRatio()) ? FVector2D(UVOffsetAmount, 0.0f) : FVector2D(0.0f, UVOffsetAmount);
+		
+		// Setup vertex buffer
+		const FVector4 Positions[] =
+		{
+			FVector4(0.0f, 1.0f, 0.0f, 1.0f),
+			FVector4(0.0f, 0.0f, 0.0f, 1.0f),
+			FVector4(1.0f, 1.0f, 0.0f, 1.0f),
+			FVector4(1.0f, 0.0f, 0.0f, 1.0f)
+		};
+
+		const FVector2D UVs[] =
+		{
+			// Landscape left
+			FVector2D(UVOffset.X, 1.0f - UVOffset.Y),
+			FVector2D(UVOffset.X, UVOffset.Y),
+			FVector2D(1.0f - UVOffset.X, 1.0f - UVOffset.Y),
+			FVector2D(1.0f - UVOffset.X, UVOffset.Y),
+
+			// Landscape right
+			FVector2D(1.0f - UVOffset.X, UVOffset.Y),
+			FVector2D(1.0f - UVOffset.X, 1.0f - UVOffset.Y),
+			FVector2D(UVOffset.X, UVOffset.Y),
+			FVector2D(UVOffset.X, 1.0f - UVOffset.Y),
+
+			// Portrait
+			FVector2D(1.0f - UVOffset.X, 1.0f - UVOffset.Y),
+			FVector2D(UVOffset.X, 1.0f - UVOffset.Y),
+			FVector2D(1.0f - UVOffset.X, UVOffset.Y),
+			FVector2D(UVOffset.X, UVOffset.Y),
+
+			// Portrait Upside Down
+			FVector2D(UVOffset.X, UVOffset.Y),
+			FVector2D(1.0f - UVOffset.X, UVOffset.Y),
+			FVector2D(UVOffset.X, 1.0f - UVOffset.Y),
+			FVector2D(1.0f - UVOffset.X, 1.0f - UVOffset.Y)
+		};
+
+		uint32 UVIndex = 0;
+		for (uint32 OrientationIter = 0; OrientationIter < 4; ++OrientationIter)
+		{
+			TResourceArray<FFilterVertex, VERTEXBUFFER_ALIGNMENT> Vertices;
+			Vertices.SetNumUninitialized(4);
+
+			Vertices[0].Position = Positions[0];
+			Vertices[0].UV = UVs[UVIndex];
+
+			Vertices[1].Position = Positions[1];
+			Vertices[1].UV = UVs[UVIndex + 1];
+
+			Vertices[2].Position = Positions[2];
+			Vertices[2].UV = UVs[UVIndex + 2];
+
+			Vertices[3].Position = Positions[3];
+			Vertices[3].UV = UVs[UVIndex + 3];
+
+			UVIndex += 4;
+
+			FRHIResourceCreateInfo CreateInfoVB(&Vertices);
+			OverlayVertexBufferRHI[OrientationIter] = RHICreateVertexBuffer(Vertices.GetResourceDataSize(), BUF_Static, CreateInfoVB);
+		}
+	}
+
+#if SUPPORTS_ARKIT_1_0
+	if (FAppleARKitAvailability::SupportsARKit10())
+	{
+		check(IsMetalPlatform(GMaxRHIShaderPlatform));
+
+		if (LastUpdateTimestamp != Frame.Timestamp && Frame.CapturedYImage && Frame.CapturedCbCrImage)
+		{
+			FRHIResourceCreateInfo CreateInfo;
+			const uint32 CreateFlags = TexCreate_Dynamic | TexCreate_NoTiling | TexCreate_ShaderResource;
+			CreateInfo.BulkData = new FAppleARKitCameraTextureResourceWrapper(Frame.CapturedYImage);
+			CreateInfo.ResourceArray = nullptr;
+
+			// pull the Y and CbCr textures out of the captured image planes (format is fake here, it will get the format from the FAppleARKitCameraTextureResourceWrapper)
+			VideoTextureY = RHICreateTexture2D(Frame.CapturedYImageWidth, Frame.CapturedYImageHeight, /*Format=*/PF_B8G8R8A8, /*NumMips=*/1, /*NumSamples=*/1, CreateFlags, CreateInfo);
+
+			CreateInfo.BulkData = new FAppleARKitCameraTextureResourceWrapper(Frame.CapturedCbCrImage);
+			VideoTextureCbCr = RHICreateTexture2D(Frame.CapturedCbCrImageWidth, Frame.CapturedCbCrImageHeight, /*Format=*/PF_B8G8R8A8, /*NumMips=*/1, /*NumSamples=*/1, CreateFlags, CreateInfo);
+
+			// todo: Add an update call to the registry instead of this unregister/re-register
+			FExternalTextureRegistry::Get().UnregisterExternalTexture(ARKitPassthroughCameraExternalTextureYGuid);
+			FExternalTextureRegistry::Get().UnregisterExternalTexture(ARKitPassthroughCameraExternalTextureCbCrGuid);
+
+			FSamplerStateInitializerRHI SamplerStateInitializer(SF_Bilinear, AM_Wrap, AM_Wrap, AM_Wrap);
+			FSamplerStateRHIRef SamplerStateRHI = RHICreateSamplerState(SamplerStateInitializer);
+
+			FExternalTextureRegistry::Get().RegisterExternalTexture(ARKitPassthroughCameraExternalTextureYGuid, VideoTextureY, SamplerStateRHI);
+			FExternalTextureRegistry::Get().RegisterExternalTexture(ARKitPassthroughCameraExternalTextureCbCrGuid, VideoTextureCbCr, SamplerStateRHI);
+
+			CFRelease(Frame.CapturedYImage);
+			CFRelease(Frame.CapturedCbCrImage);
+			Frame.CapturedYImage = nullptr;
+			Frame.CapturedCbCrImage = nullptr;
+
+			LastUpdateTimestamp = Frame.Timestamp;
+		}
+	}
+#endif
+}
+
+// We use something similar to the PostProcessMaterial to render the color camera overlay.
+template <bool bIsMobileRenderer>
+class FARKitCameraOverlayVS : public FMaterialShader
+{
+	DECLARE_SHADER_TYPE(FARKitCameraOverlayVS, Material);
+public:
+
+	static bool ShouldCompilePermutation(EShaderPlatform Platform, const FMaterial* Material)
+	{
+		if (bIsMobileRenderer)
+		{
+			return Material->GetMaterialDomain() == MD_PostProcess && IsMobilePlatform(Platform);
+		}
+		else
+		{
+			return Material->GetMaterialDomain() == MD_PostProcess && !IsMobilePlatform(Platform);
+		}
+	}
+
+	static void ModifyCompilationEnvironment(EShaderPlatform Platform, const class FMaterial* Material, FShaderCompilerEnvironment& OutEnvironment)
+	{
+		FMaterialShader::ModifyCompilationEnvironment(Platform, OutEnvironment);
+		OutEnvironment.SetDefine(TEXT("POST_PROCESS_MATERIAL"), 1);
+		OutEnvironment.SetDefine(TEXT("POST_PROCESS_MATERIAL_BEFORE_TONEMAP"), (Material->GetBlendableLocation() != BL_AfterTonemapping) ? 1 : 0);
+	}
+
+	FARKitCameraOverlayVS() { }
+	FARKitCameraOverlayVS(const ShaderMetaType::CompiledShaderInitializerType& Initializer)
+		: FMaterialShader(Initializer)
+	{
+	}
+
+	void SetParameters(FRHICommandList& RHICmdList, const FSceneView View)
+	{
+		const FVertexShaderRHIParamRef ShaderRHI = GetVertexShader();
+		FMaterialShader::SetViewParameters(RHICmdList, ShaderRHI, View, View.ViewUniformBuffer);
+	}
+
+	virtual bool Serialize(FArchive& Ar) override
+	{
+		const bool bShaderHasOutdatedParameters = FMaterialShader::Serialize(Ar);
+		return bShaderHasOutdatedParameters;
+	}
+};
+
+IMPLEMENT_MATERIAL_SHADER_TYPE(template<>, FARKitCameraOverlayVS<true>, TEXT("/Engine/Private/PostProcessMaterialShaders.usf"), TEXT("MainVS_ES2"), SF_Vertex);
+IMPLEMENT_MATERIAL_SHADER_TYPE(template<>, FARKitCameraOverlayVS<false>, TEXT("/Engine/Private/PostProcessMaterialShaders.usf"), TEXT("MainVS_VideoOverlay"), SF_Vertex);
+
+template <bool bIsMobileRenderer>
+class FARKitCameraOverlayPS : public FMaterialShader
+{
+	DECLARE_SHADER_TYPE(FARKitCameraOverlayPS, Material);
+public:
+
+	static bool ShouldCompilePermutation(EShaderPlatform Platform, const FMaterial* Material)
+	{
+		if (bIsMobileRenderer)
+		{
+			return Material->GetMaterialDomain() == MD_PostProcess && IsMobilePlatform(Platform);
+		}
+		else
+		{
+			return Material->GetMaterialDomain() == MD_PostProcess && !IsMobilePlatform(Platform);
+		}
+	}
+
+	static void ModifyCompilationEnvironment(EShaderPlatform Platform, const class FMaterial* Material, FShaderCompilerEnvironment& OutEnvironment)
+	{
+		FMaterialShader::ModifyCompilationEnvironment(Platform, OutEnvironment);
+		OutEnvironment.SetDefine(TEXT("POST_PROCESS_MATERIAL"), 1);
+		OutEnvironment.SetDefine(TEXT("OUTPUT_GAMMA_SPACE"), IsMobileHDR() ? 0 : 1);
+		OutEnvironment.SetDefine(TEXT("POST_PROCESS_MATERIAL_BEFORE_TONEMAP"), (Material->GetBlendableLocation() != BL_AfterTonemapping) ? 1 : 0);
+	}
+
+	FARKitCameraOverlayPS() {}
+	FARKitCameraOverlayPS(const ShaderMetaType::CompiledShaderInitializerType& Initializer) :
+		FMaterialShader(Initializer)
+	{
+		for (uint32 InputIter = 0; InputIter < ePId_Input_MAX; ++InputIter)
+		{
+			PostprocessInputParameter[InputIter].Bind(Initializer.ParameterMap, *FString::Printf(TEXT("PostprocessInput%d"), InputIter));
+			PostprocessInputParameterSampler[InputIter].Bind(Initializer.ParameterMap, *FString::Printf(TEXT("PostprocessInput%dSampler"), InputIter));
+		}
+	}
+
+	void SetParameters(FRHICommandList& RHICmdList, const FSceneView View, const FMaterialRenderProxy* Material)
+	{
+		const FPixelShaderRHIParamRef ShaderRHI = GetPixelShader();
+		FMaterialShader::SetParameters(RHICmdList, ShaderRHI, Material, *Material->GetMaterial(View.GetFeatureLevel()), View, View.ViewUniformBuffer, ESceneTextureSetupMode::None);
+
+		for (uint32 InputIter = 0; InputIter < ePId_Input_MAX; ++InputIter)
+		{
+			if (PostprocessInputParameter[InputIter].IsBound())
+			{
+				SetTextureParameter(
+					RHICmdList, 
+					ShaderRHI, 
+					PostprocessInputParameter[InputIter],
+					PostprocessInputParameterSampler[InputIter],
+					TStaticSamplerState<>::GetRHI(),
+					GBlackTexture->TextureRHI);
+			}
+		}
+	}
+
+	virtual bool Serialize(FArchive& Ar) override
+	{
+		const bool bShaderHasOutdatedParameters = FMaterialShader::Serialize(Ar);
+		return bShaderHasOutdatedParameters;
+	}
+
+private:
+	FShaderResourceParameter PostprocessInputParameter[ePId_Input_MAX];
+	FShaderResourceParameter PostprocessInputParameterSampler[ePId_Input_MAX];
+};
+
+IMPLEMENT_MATERIAL_SHADER_TYPE(template<>, FARKitCameraOverlayPS<true>, TEXT("/Engine/Private/PostProcessMaterialShaders.usf"), TEXT("MainPS_ES2"), SF_Pixel);
+IMPLEMENT_MATERIAL_SHADER_TYPE(template<>, FARKitCameraOverlayPS<false>, TEXT("/Engine/Private/PostProcessMaterialShaders.usf"), TEXT("MainPS_VideoOverlay"), SF_Pixel);
+
+void FAppleARKitVideoOverlay::RenderVideoOverlay_RenderThread(FRHICommandListImmediate& RHICmdList, const FSceneView& InView, const EScreenOrientation::Type DeviceOrientation)
+{
+#if SUPPORTS_ARKIT_1_0
+
+	if (FAppleARKitAvailability::SupportsARKit10())
+	{
+		if (RenderingOverlayMaterial == nullptr || !RenderingOverlayMaterial->IsValidLowLevel())
+		{
+			return;
+		}
+
+		const auto FeatureLevel = InView.GetFeatureLevel();
+		IRendererModule& RendererModule = GetRendererModule();
+
+		const FMaterial* const CameraMaterial = RenderingOverlayMaterial->GetRenderProxy(false)->GetMaterial(FeatureLevel);
+		const FMaterialShaderMap* const MaterialShaderMap = CameraMaterial->GetRenderingThreadShaderMap();
+
+		FGraphicsPipelineStateInitializer GraphicsPSOInit;
+		RHICmdList.ApplyCachedRenderTargets(GraphicsPSOInit);
+
+		GraphicsPSOInit.BlendState = TStaticBlendState<>::GetRHI();
+		GraphicsPSOInit.RasterizerState = TStaticRasterizerState<>::GetRHI();
+		GraphicsPSOInit.DepthStencilState = TStaticDepthStencilState<false, CF_DepthNearOrEqual>::GetRHI();
+		GraphicsPSOInit.PrimitiveType = PT_TriangleList;
+		GraphicsPSOInit.BoundShaderState.VertexDeclarationRHI = RendererModule.GetFilterVertexDeclaration().VertexDeclarationRHI;
+		
+		const bool bIsMobileRenderer = FeatureLevel <= ERHIFeatureLevel::ES3_1;
+		FMaterialShader* VertexShader = nullptr;
+		FMaterialShader* PixelShader = nullptr;
+		if (bIsMobileRenderer)
+		{
+			VertexShader = MaterialShaderMap->GetShader<FARKitCameraOverlayVS<true>>();
+			PixelShader = MaterialShaderMap->GetShader<FARKitCameraOverlayPS<true>>();
+			GraphicsPSOInit.BoundShaderState.VertexShaderRHI = GETSAFERHISHADER_VERTEX(reinterpret_cast<FARKitCameraOverlayVS<true>*>(VertexShader));
+			GraphicsPSOInit.BoundShaderState.PixelShaderRHI = GETSAFERHISHADER_PIXEL(reinterpret_cast<FARKitCameraOverlayPS<true>*>(PixelShader));
+		}
+		else
+		{
+			VertexShader = MaterialShaderMap->GetShader<FARKitCameraOverlayVS<false>>();
+			PixelShader = MaterialShaderMap->GetShader<FARKitCameraOverlayPS<false>>();
+			GraphicsPSOInit.BoundShaderState.VertexShaderRHI = GETSAFERHISHADER_VERTEX(reinterpret_cast<FARKitCameraOverlayVS<false>*>(VertexShader));
+			GraphicsPSOInit.BoundShaderState.PixelShaderRHI = GETSAFERHISHADER_PIXEL(reinterpret_cast<FARKitCameraOverlayPS<false>*>(PixelShader));
+		}
+		
+		SetGraphicsPipelineState(RHICmdList, GraphicsPSOInit);
+		
+		const FIntPoint ViewSize = InView.UnconstrainedViewRect.Size();
+		FDrawRectangleParameters Parameters;
+		Parameters.PosScaleBias = FVector4(ViewSize.X, ViewSize.Y, 0, 0);
+		Parameters.UVScaleBias = FVector4(1.0f, 1.0f, 0.0f, 0.0f);
+		Parameters.InvTargetSizeAndTextureSize = FVector4(
+														  1.0f / ViewSize.X, 1.0f / ViewSize.Y,
+														  1.0f, 1.0f);
+
+		if (bIsMobileRenderer)
+		{
+			FARKitCameraOverlayVS<true>* const VertexShaderPtr = reinterpret_cast<FARKitCameraOverlayVS<true>*>(VertexShader);
+			SetUniformBufferParameterImmediate(RHICmdList, VertexShaderPtr->GetVertexShader(), VertexShaderPtr->GetUniformBufferParameter<FDrawRectangleParameters>(), Parameters);
+			VertexShaderPtr->SetParameters(RHICmdList, InView);
+			reinterpret_cast<FARKitCameraOverlayPS<true>*>(PixelShader)->SetParameters(RHICmdList, InView, RenderingOverlayMaterial->GetRenderProxy(false));
+		}
+		else
+		{
+			FARKitCameraOverlayVS<false>* const VertexShaderPtr = reinterpret_cast<FARKitCameraOverlayVS<false>*>(VertexShader);
+			SetUniformBufferParameterImmediate(RHICmdList, VertexShaderPtr->GetVertexShader(), VertexShaderPtr->GetUniformBufferParameter<FDrawRectangleParameters>(), Parameters);
+			VertexShaderPtr->SetParameters(RHICmdList, InView);
+			reinterpret_cast<FARKitCameraOverlayPS<false>*>(PixelShader)->SetParameters(RHICmdList, InView, RenderingOverlayMaterial->GetRenderProxy(false));
+		}
+
+		FVertexBufferRHIParamRef VertexBufferRHI = nullptr;
+		switch (DeviceOrientation)
+		{
+			case EScreenOrientation::Type::LandscapeLeft:
+				VertexBufferRHI = OverlayVertexBufferRHI[0];
+				break;
+
+			case EScreenOrientation::Type::LandscapeRight:
+				VertexBufferRHI = OverlayVertexBufferRHI[1];
+				break;
+
+			case EScreenOrientation::Type::Portrait:
+				VertexBufferRHI = OverlayVertexBufferRHI[2];
+				break;
+
+			case EScreenOrientation::Type::PortraitUpsideDown:
+				VertexBufferRHI = OverlayVertexBufferRHI[3];
+				break;
+
+			default:
+				VertexBufferRHI = OverlayVertexBufferRHI[0];
+				break;
+		}
+
+		if (VertexBufferRHI && OverlayIndexBufferRHI.IsValid())
+		{
+			RHICmdList.SetStreamSource(0, VertexBufferRHI, 0);
+			RHICmdList.DrawIndexedPrimitive(
+				OverlayIndexBufferRHI,
+				PT_TriangleList,
+				/*BaseVertexIndex=*/ 0,
+				/*MinIndex=*/ 0,
+				/*NumVertices=*/ 4,
+				/*StartIndex=*/ 0,
+				/*NumPrimitives=*/ 2,
+				/*NumInstances=*/ 1
+			);
+		}
+	}
+#endif
+}
diff --git a/Engine/Plugins/Experimental/AR/Apple/AppleARKit/Source/AppleARKit/Private/AppleARKitVideoOverlay.h b/Engine/Plugins/Experimental/AR/Apple/AppleARKit/Source/AppleARKit/Private/AppleARKitVideoOverlay.h
new file mode 100644
--- /dev/null
+++ b/Engine/Plugins/Experimental/AR/Apple/AppleARKit/Source/AppleARKit/Private/AppleARKitVideoOverlay.h
@@ -0,0 +1,50 @@
+// Copyright 1998-2018 Epic Games, Inc. All Rights Reserved.
+
+#pragma once
+
+#include "CoreMinimal.h"
+#include "Materials/MaterialInterface.h"
+#include "UObject/ConstructorHelpers.h"
+#include "Kismet/BlueprintPlatformLibrary.h"
+#include "AppleARKitVideoOverlay.generated.h"
+
+class FSceneViewFamily;
+
+/** Helper class to ensure the ARKit camera material is cooked. */
+UCLASS()
+class UARKitCameraOverlayMaterialLoader : public UObject
+{
+	GENERATED_BODY()
+	
+public:
+	UPROPERTY()
+	UMaterialInterface* DefaultCameraOverlayMaterial;
+	
+	UARKitCameraOverlayMaterialLoader()
+	{
+		static ConstructorHelpers::FObjectFinder<UMaterialInterface> DefaultOverlayMaterialRef(TEXT("/AppleARKit/ARKitCameraMaterial.ARKitCameraMaterial"));
+		DefaultCameraOverlayMaterial = DefaultOverlayMaterialRef.Object;
+	}
+};
+
+struct FAppleARKitFrame;
+
+class FAppleARKitVideoOverlay
+{
+public:
+	FAppleARKitVideoOverlay();
+
+	void UpdateVideoTexture_RenderThread(FRHICommandListImmediate& RHICmdList, FAppleARKitFrame& Frame, const FSceneViewFamily& InViewFamily);
+	void RenderVideoOverlay_RenderThread(FRHICommandListImmediate& RHICmdList, const FSceneView& InView, const EScreenOrientation::Type DeviceOrientation);
+
+private:
+	FTextureRHIRef VideoTextureY;
+	FTextureRHIRef VideoTextureCbCr;
+	UMaterialInterface* RenderingOverlayMaterial;
+	FIndexBufferRHIRef OverlayIndexBufferRHI;
+	
+	// Separate vertex buffer for each supported device orientation
+	FVertexBufferRHIRef OverlayVertexBufferRHI[4];
+
+	double LastUpdateTimestamp;
+};
diff --git a/Engine/Plugins/Experimental/AR/Apple/AppleARKit/Source/AppleARKit/Private/ExternalTextureGuid.h b/Engine/Plugins/Experimental/AR/Apple/AppleARKit/Source/AppleARKit/Private/ExternalTextureGuid.h
new file mode 100644
--- /dev/null
+++ b/Engine/Plugins/Experimental/AR/Apple/AppleARKit/Source/AppleARKit/Private/ExternalTextureGuid.h
@@ -0,0 +1,15 @@
+// Copyright 1998-2018 Epic Games, Inc. All Rights Reserved.
+
+#pragma once
+
+#include "CoreMinimal.h"
+#include "Misc/Guid.h"
+
+/* GUID used to indicate that this external texture is the ARKit camera Y component. 0a417e66-bcda-4610-8594-92248c914a20 */
+static const FGuid ARKitPassthroughCameraExternalTextureYGuid(0x0A417E66, 0xBCDA4610, 0x85949224, 0x8C914A20);
+
+/* GUID used to indicate that this external texture is the ARKit camera CbCr component. d61a2708-55fe-4cb1-9cde-df2a9bb2a2d2 */
+static const FGuid ARKitPassthroughCameraExternalTextureCbCrGuid(0xD61A2708, 0x55FE4CB1, 0x9CDEDF2A, 0x9BB2A2D2);
+
+/* GUID used to indicate that this external texture is the ARKit camera depth texture. cba06c09-03b2-4818-8913-6424d99acf28 */
+static const FGuid ARKitPassthroughCameraExternalTextureDepthGuid(0xCBA06C09, 0x03B24818, 0x89136424, 0xD99ACF28);
diff --git a/Engine/Plugins/Experimental/AR/Apple/AppleARKit/Source/AppleARKit/Private/ExternalTextureMaterialExpression.cpp b/Engine/Plugins/Experimental/AR/Apple/AppleARKit/Source/AppleARKit/Private/ExternalTextureMaterialExpression.cpp
new file mode 100644
--- /dev/null
+++ b/Engine/Plugins/Experimental/AR/Apple/AppleARKit/Source/AppleARKit/Private/ExternalTextureMaterialExpression.cpp
@@ -0,0 +1,29 @@
+// Copyright 1998-2018 Epic Games, Inc. All Rights Reserved.
+
+#include "ExternalTextureMaterialExpression.h"
+#include "MaterialCompiler.h"
+#include "ExternalTextureGuid.h"
+
+UMaterialExpressionARKitPassthroughCamera::UMaterialExpressionARKitPassthroughCamera(const FObjectInitializer& ObjectInitializer)
+	: Super(ObjectInitializer)
+{}
+
+#if WITH_EDITOR
+int32 UMaterialExpressionARKitPassthroughCamera::Compile(class FMaterialCompiler* Compiler, int32 OutputIndex)
+{
+	return Compiler->TextureSample(
+		Compiler->ExternalTexture((TextureType == TextureY) ? ARKitPassthroughCameraExternalTextureYGuid : ARKitPassthroughCameraExternalTextureCbCrGuid),
+		Coordinates.GetTracedInput().Expression ? Coordinates.Compile(Compiler) : Compiler->TextureCoordinate(ConstCoordinate, false, false),
+		EMaterialSamplerType::SAMPLERTYPE_Color);
+}
+
+int32 UMaterialExpressionARKitPassthroughCamera::CompilePreview(class FMaterialCompiler* Compiler, int32 OutputIndex)
+{
+	return INDEX_NONE;
+}
+
+void UMaterialExpressionARKitPassthroughCamera::GetCaption(TArray<FString>& OutCaptions) const
+{
+	OutCaptions.Add(TEXT("ARKit Passthrough Camera"));
+}
+#endif
\ No newline at end of file
diff --git a/Engine/Plugins/Experimental/AR/Apple/AppleARKit/Source/AppleARKit/Public/AppleARKitAnchor.h b/Engine/Plugins/Experimental/AR/Apple/AppleARKit/Source/AppleARKit/Public/AppleARKitAnchor.h
new file mode 100644
--- /dev/null
+++ b/Engine/Plugins/Experimental/AR/Apple/AppleARKit/Source/AppleARKit/Public/AppleARKitAnchor.h
@@ -0,0 +1,59 @@
+// Copyright 1998-2018 Epic Games, Inc. All Rights Reserved.
+
+#pragma once
+
+#include "CoreMinimal.h"
+#include "AppleARKitAvailability.h"
+
+// ARKit
+#if SUPPORTS_ARKIT_1_0
+	#import <ARKit/ARKit.h>
+#endif
+
+// UE4
+#include "Misc/Guid.h"
+#include "HAL/CriticalSection.h"
+
+// AppleARKit
+#include "AppleARKitAnchor.generated.h"
+
+UCLASS( BlueprintType, Deprecated )
+class APPLEARKIT_API UDEPRECATED_AppleARKitAnchor : public UObject
+{
+	GENERATED_BODY()
+
+public: 
+
+	/**
+	 * Unique identifier of the anchor.
+	 */
+	UPROPERTY()
+	FGuid Identifier;
+
+	/**
+	 * The transformation matrix that defines the anchor's rotation, translation and scale.
+	 *
+	 * NOTE: This does not have Session::BaseTransform applied due to thread safety issues. You'll
+	 * need to apply this yourself in the game thread.
+	 *
+	 * @todo arkit Fix this ^
+	 */
+	UFUNCTION( BlueprintPure, Category = "AppleARKit|Anchor" )
+	FTransform GetTransform() const;
+
+#if SUPPORTS_ARKIT_1_0
+
+	virtual void Update_DelegateThread( ARAnchor* Anchor );
+
+#endif
+
+protected:
+
+	// Thread safe update lock
+	mutable FCriticalSection UpdateLock;
+
+	/**
+	 * The transformation matrix that defines the anchor's rotation, translation and scale in world coordinates.
+	 */
+	FTransform Transform;
+};
diff --git a/Engine/Plugins/Experimental/AR/Apple/AppleARKit/Source/AppleARKit/Public/AppleARKitAvailability.h b/Engine/Plugins/Experimental/AR/Apple/AppleARKit/Source/AppleARKit/Public/AppleARKitAvailability.h
new file mode 100644
--- /dev/null
+++ b/Engine/Plugins/Experimental/AR/Apple/AppleARKit/Source/AppleARKit/Public/AppleARKitAvailability.h
@@ -0,0 +1,99 @@
+// Copyright 1998-2018 Epic Games, Inc. All Rights Reserved.
+
+#pragma once
+
+
+#include "CoreMinimal.h"
+#if PLATFORM_IOS && !PLATFORM_TVOS
+	#include "Availability.h"
+#endif
+
+#if PLATFORM_IOS && !PLATFORM_TVOS
+
+	// Check for ARKit 1.0
+	#ifdef __IPHONE_11_0
+		#define SUPPORTS_ARKIT_1_0 1
+	#else
+		#define SUPPORTS_ARKIT_1_0 0
+	#endif
+
+	// Check for ARKit 1.5
+	#ifdef __IPHONE_11_3
+		#define SUPPORTS_ARKIT_1_5 1
+	#else
+		#define SUPPORTS_ARKIT_1_5 0
+	#endif
+
+	// Check for ARKit 2.0
+	#ifdef __IPHONE_12_0
+		#define SUPPORTS_ARKIT_2_0 1
+	#else
+		#define SUPPORTS_ARKIT_2_0 0
+	#endif
+
+#else
+
+	// No ARKit support
+	#define SUPPORTS_ARKIT_1_0 0
+	#define SUPPORTS_ARKIT_1_5 0
+	#define SUPPORTS_ARKIT_2_0 0
+
+#endif
+
+class FAppleARKitAvailability
+{
+public:
+	static bool SupportsARKit10()
+	{
+		static bool bSupportsARKit10 = false;
+#if SUPPORTS_ARKIT_1_0
+		static bool bSupportChecked = false;
+		if (!bSupportChecked)
+		{
+			bSupportChecked = true;
+			// This call is slow, so cache the result
+			if (@available(iOS 11.0, *))
+			{
+				bSupportsARKit10 = true;
+			}
+		}
+#endif
+		return bSupportsARKit10;
+	}
+
+	static bool SupportsARKit15()
+	{
+		static bool bSupportsARKit15 = false;
+#if SUPPORTS_ARKIT_1_5
+		static bool bSupportChecked = false;
+		if (!bSupportChecked)
+		{
+			bSupportChecked = true;
+			// This call is slow, so cache the result
+			if (@available(iOS 11.3, *))
+			{
+				bSupportsARKit15 = true;
+			}
+		}
+#endif
+		return bSupportsARKit15;
+	}
+
+	static bool SupportsARKit20()
+	{
+		static bool bSupportsARKit20 = false;
+#if SUPPORTS_ARKIT_2_0
+		static bool bSupportChecked = false;
+		if (!bSupportChecked)
+		{
+			bSupportChecked = true;
+			// This call is slow, so cache the result
+			if (@available(iOS 12.0, *))
+			{
+				bSupportsARKit20 = true;
+			}
+		}
+#endif
+		return bSupportsARKit20;
+	}
+};
diff --git a/Engine/Plugins/Experimental/AR/Apple/AppleARKit/Source/AppleARKit/Public/AppleARKitBlueprintLibrary.h b/Engine/Plugins/Experimental/AR/Apple/AppleARKit/Source/AppleARKit/Public/AppleARKitBlueprintLibrary.h
new file mode 100644
--- /dev/null
+++ b/Engine/Plugins/Experimental/AR/Apple/AppleARKit/Source/AppleARKit/Public/AppleARKitBlueprintLibrary.h
@@ -0,0 +1,42 @@
+// Copyright 1998-2018 Epic Games, Inc. All Rights Reserved.
+
+#pragma once
+
+#include "CoreMinimal.h"
+
+#include "AppleARKitFrame.h"
+#include "AppleARKitHitTestResult.h"
+
+#include "Kismet/BlueprintFunctionLibrary.h"
+#include "AppleARKitBlueprintLibrary.generated.h"
+
+UCLASS(meta=(ScriptName="AppleARKitLibrary"))
+class APPLEARKIT_API UAppleARKitBlueprintLibrary : public UBlueprintFunctionLibrary
+{
+	GENERATED_BODY()
+
+public:
+	/**
+	 * Returns the last processed FAppleARKitFrame from the session.
+	 *
+	 * When called from the game thread, we ensure a single frame is returned for the duration
+	 * of the game thread update by watching for changes to GFrameNumber to trigger pulling a
+	 * new frame from the session.
+	 */
+	UFUNCTION( BlueprintCallable, Category="AppleARKit", meta = (DeprecatedFunction, WorldContext = "WorldContextObject", Keywords = "ar augmentedreality augmented reality frame"))
+	static bool GetCurrentFrame( UObject* WorldContextObject, FAppleARKitFrame& OutCurrentFrame );
+	
+	/**
+	 * Searches the last processed frame for anchors corresponding to a point in the captured image.
+	 *
+	 * A 2D point in the captured image's coordinate space can refer to any point along a line segment
+	 * in the 3D coordinate space. Hit-testing is the process of finding anchors of a frame located along this line segment.
+	 *
+	 * NOTE: The hit test locations are reported in ARKit space. For hit test results
+	 * in game world coordinates, you're after UAppleARKitCameraComponent::HitTestAtScreenPosition
+	 *
+	 * @param ScreenPosition The viewport pixel coordinate of the trace origin.
+	 */
+	UFUNCTION( BlueprintCallable, Category="AppleARKit", meta = (DeprecatedFunction, WorldContext = "WorldContextObject", Keywords = "ar augmentedreality augmented reality trace hittest hit line"))
+	static bool HitTestAtScreenPosition_TrackingSpace( UObject* WorldContextObject, const FVector2D ScreenPosition, EAppleARKitHitTestResultType Types, TArray< FAppleARKitHitTestResult >& OutResults );
+};
diff --git a/Engine/Plugins/Experimental/AR/Apple/AppleARKit/Source/AppleARKit/Public/AppleARKitFrame.h b/Engine/Plugins/Experimental/AR/Apple/AppleARKit/Source/AppleARKit/Public/AppleARKitFrame.h
new file mode 100644
--- /dev/null
+++ b/Engine/Plugins/Experimental/AR/Apple/AppleARKit/Source/AppleARKit/Public/AppleARKitFrame.h
@@ -0,0 +1,89 @@
+// Copyright 1998-2018 Epic Games, Inc. All Rights Reserved.
+
+#pragma once
+
+// AppleARKit
+#include "AppleARKitAvailability.h"
+#include "AppleARKitCamera.h"
+#include "AppleARKitLightEstimate.h"
+
+#include "AppleARKitFrame.generated.h"
+
+/**
+ * An object representing a frame processed by FAppleARKitSystem.
+ * @discussion Each frame contains information about the current state of the scene.
+ */
+USTRUCT( BlueprintType, Category="AppleARKit" )
+struct APPLEARKIT_API FAppleARKitFrame
+{
+	GENERATED_BODY()
+	
+	// Default constructor
+	FAppleARKitFrame();
+
+#if SUPPORTS_ARKIT_1_0
+
+	/** 
+	 * This is a conversion copy-constructor that takes a raw ARFrame and fills this structs members
+	 * with the UE4-ified versions of ARFrames properties.
+	 *
+	 * @param InARFrame 		- The frame to convert / copy
+	 * @param MetalTextureCache - A CVMetalTextureCacheRef to use in the conversion of InARFrame's CVPixelBufferRef to CVMetalTextureRef for our CapturedImage property.
+	 */ 
+	FAppleARKitFrame( ARFrame* InARFrame, CVMetalTextureCacheRef MetalTextureCache );
+
+	/** 
+	 * Copy contructor. CapturedImage is skipped as we don't need / want to retain access to the 
+	 * image buffer.
+	 */
+	FAppleARKitFrame( const FAppleARKitFrame& Other );
+
+	/** Destructor. Calls CFRelease on CapturedImage */
+	virtual ~FAppleARKitFrame();
+
+	/** 
+	 * Assignment operator. CapturedImage is skipped as we don't need / want to retain access to the 
+	 * image buffer. */
+	FAppleARKitFrame& operator=( const FAppleARKitFrame& Other );
+
+#endif
+
+	/** A timestamp identifying the frame. */
+	double Timestamp;
+
+#if SUPPORTS_ARKIT_1_0
+
+	/** The frame's captured images */
+    CVMetalTextureRef CapturedYImage;
+    CVMetalTextureRef CapturedCbCrImage;
+
+	/** The raw camera buffer from ARKit */
+	CVPixelBufferRef CameraImage;
+	/** The raw camera depth info from ARKit (needs iPhone X) */
+	AVDepthData* CameraDepth;
+	void* NativeFrame;
+#endif
+
+	/** The width and height in pixels of the frame's captured images. */
+	uint32 CapturedYImageWidth;
+	uint32 CapturedYImageHeight;
+    
+    uint32 CapturedCbCrImageWidth;
+    uint32 CapturedCbCrImageHeight;
+    
+	/** The camera used to capture the frame's image. */
+	FAppleARKitCamera Camera;
+
+	/**
+	 * A light estimate representing the estimated light in the scene.
+	 */
+	FAppleARKitLightEstimate LightEstimate;
+	
+	/** The current world mapping state is reported on the frame */
+	EARWorldMappingState WorldMappingState;
+
+	/* 
+	 * When adding new member variables, don't forget to handle them in the copy constructor and
+	 * assignment operator above.
+	 */
+};
diff --git a/Engine/Plugins/Experimental/AR/Apple/AppleARKit/Source/AppleARKit/Public/AppleARKitHitTestResult.h b/Engine/Plugins/Experimental/AR/Apple/AppleARKit/Source/AppleARKit/Public/AppleARKitHitTestResult.h
new file mode 100644
--- /dev/null
+++ b/Engine/Plugins/Experimental/AR/Apple/AppleARKit/Source/AppleARKit/Public/AppleARKitHitTestResult.h
@@ -0,0 +1,95 @@
+// Copyright 1998-2018 Epic Games, Inc. All Rights Reserved.
+
+#pragma once
+
+#include "CoreMinimal.h"
+#include "AppleARKitAvailability.h"
+
+#if SUPPORTS_ARKIT_1_0
+	#import <ARKit/ARKit.h>
+#endif
+
+
+#include "ARSystem.h"
+#include "AppleARKitHitTestResult.generated.h"
+
+/**
+ * Option set of hit-test result types.
+ */
+UENUM( BlueprintType, Category="AppleARKit", meta=(Bitflags) )
+enum class EAppleARKitHitTestResultType : uint8
+{
+	None = 0,
+
+    /** Result type from intersecting the nearest feature point. */
+    FeaturePoint = 1 UMETA(DisplayName = "Feature Point"),
+    
+    /** Result type from intersecting a horizontal plane estimate, determined for the current frame. */
+    EstimatedHorizontalPlane = 2 UMETA(DisplayName = "Estimated Horizontal Plane"),
+    
+    /** Result type from intersecting with an existing plane anchor. */
+    ExistingPlane = 4 UMETA(DisplayName = "Existing Plane"),
+    
+    /** Result type from intersecting with an existing plane anchor. */
+    ExistingPlaneUsingExtent = 8 UMETA(DisplayName= "Existing Plane Using Extent")
+};
+ENUM_CLASS_FLAGS(EAppleARKitHitTestResultType);
+
+#if SUPPORTS_ARKIT_1_0
+
+/** Conversion function from ARKit native ARHitTestResultType */
+EAppleARKitHitTestResultType ToEAppleARKitHitTestResultType(ARHitTestResultType InTypes);
+
+#endif
+
+/**
+ * A result of an intersection found during a hit-test.
+ */
+USTRUCT( BlueprintType, Category="AppleARKit")
+struct APPLEARKIT_API FAppleARKitHitTestResult
+{
+    GENERATED_BODY();
+
+    // Default constructor
+	FAppleARKitHitTestResult()
+		: Type(EAppleARKitHitTestResultType::None)
+		, Distance(0.f)
+	{};
+
+#if SUPPORTS_ARKIT_1_0
+
+	/** 
+	 * This is a conversion copy-constructor that takes a raw ARHitTestResult and fills this 
+	 * structs members with the UE4-ified versions of ARHitTestResult's properties.
+	 */ 
+	FAppleARKitHitTestResult( ARHitTestResult* InARHitTestResult, class UDEPRECATED_AppleARKitAnchor* InAnchor = nullptr, float WorldToMetersScale = 100.0f );
+
+#endif
+
+	/**
+	 * The type of the hit-test result.
+	 */
+    UPROPERTY( BlueprintReadOnly, Category = "AppleARKitHitTestResult" )
+    EAppleARKitHitTestResultType Type;
+    
+	/**
+	 * The distance from the camera to the intersection in meters.
+	 */
+    UPROPERTY( BlueprintReadOnly, Category = "AppleARKitHitTestResult")
+	float Distance;
+
+	/**
+	 * The transformation matrix that defines the intersection's rotation, translation and scale
+	 * relative to the world.
+	 */
+    UPROPERTY( BlueprintReadOnly, Category = "AppleARKitHitTestResult")
+	FTransform Transform;
+
+	/**
+	 * The anchor that the hit-test intersected.
+	 * 
+	 * An anchor will only be provided for existing plane result types.
+	 */
+	UPROPERTY()
+	class UDEPRECATED_AppleARKitAnchor* Anchor_DEPRECATED = nullptr;
+};
diff --git a/Engine/Plugins/Experimental/AR/Apple/AppleARKit/Source/AppleARKit/Public/AppleARKitPlaneAnchor.h b/Engine/Plugins/Experimental/AR/Apple/AppleARKit/Source/AppleARKit/Public/AppleARKitPlaneAnchor.h
new file mode 100644
--- /dev/null
+++ b/Engine/Plugins/Experimental/AR/Apple/AppleARKit/Source/AppleARKit/Public/AppleARKitPlaneAnchor.h
@@ -0,0 +1,56 @@
+// Copyright 1998-2018 Epic Games, Inc. All Rights Reserved.
+
+#pragma once
+
+#include "CoreMinimal.h"
+#include "AppleARKitAvailability.h"
+
+// ARKit
+#if SUPPORTS_ARKIT_1_0
+	#import <ARKit/ARKit.h>
+#endif
+
+// AppleARKit
+#include "AppleARKitAnchor.h"
+#include "AppleARKitPlaneAnchor.generated.h"
+
+UCLASS( BlueprintType, Deprecated )
+class APPLEARKIT_API UDEPRECATED_AppleARKitPlaneAnchor : public UDEPRECATED_AppleARKitAnchor
+{
+	GENERATED_BODY()
+
+public: 
+
+	/**
+	 * The center of the plane in the anchor’s coordinate space.
+	 */
+	UFUNCTION( BlueprintPure, Category = "AppleARKit|PlaneAnchor" )
+	FVector GetCenter() const;
+
+	/**
+	 * The extent of the plane in the anchor’s coordinate space.
+	 */
+	UFUNCTION( BlueprintPure, Category = "AppleARKit|PlaneAnchor")
+	FVector GetExtent() const;
+
+	UFUNCTION( BlueprintPure, Category = "AppleARKit|PlaneAnchor")
+	FTransform GetTransformToCenter() const;
+
+#if SUPPORTS_ARKIT_1_0
+
+	virtual void Update_DelegateThread( ARAnchor* Anchor ) override;
+
+#endif
+
+private:
+
+	/**
+	 * The center of the plane in the anchor’s coordinate space.
+	 */
+	FVector Center;
+
+	/**
+	 * The extent of the plane in the anchor’s coordinate space.
+	 */
+	FVector Extent;
+};
diff --git a/Engine/Plugins/Experimental/AR/Apple/AppleARKit/Source/AppleARKit/Public/AppleARKitSettings.h b/Engine/Plugins/Experimental/AR/Apple/AppleARKit/Source/AppleARKit/Public/AppleARKitSettings.h
new file mode 100644
--- /dev/null
+++ b/Engine/Plugins/Experimental/AR/Apple/AppleARKit/Source/AppleARKit/Public/AppleARKitSettings.h
@@ -0,0 +1,55 @@
+// Copyright 1998-2018 Epic Games, Inc. All Rights Reserved.
+
+#pragma once
+
+#include "UObject/ObjectMacros.h"
+#include "UObject/Object.h"
+#include "ARTrackable.h"
+
+#include "AppleARKitSettings.generated.h"
+
+UCLASS(Config=Engine)
+class APPLEARKIT_API UAppleARKitSettings : public UObject
+{
+	GENERATED_BODY()
+
+public:
+	UAppleARKitSettings()
+		: bEnableLiveLinkForFaceTracking(false)
+		, LiveLinkPublishingPort(11111)
+		, DefaultFaceTrackingLiveLinkSubjectName(FName("iPhoneXFaceAR"))
+		, DefaultFaceTrackingDirection(EARFaceTrackingDirection::FaceRelative)
+		, bAdjustThreadPrioritiesDuringARSession(false)
+		, GameThreadPriorityOverride(47)
+		, RenderThreadPriorityOverride(45)
+	{
+	}
+
+	/** Whether to publish face blend shapes to LiveLink or not */
+	UPROPERTY(Config, EditAnywhere, Category="AR Settings")
+	bool bEnableLiveLinkForFaceTracking;
+
+	/** The port to use when listening/sending LiveLink face blend shapes via the network */
+	UPROPERTY(Config, EditAnywhere, Category="AR Settings")
+	int32 LiveLinkPublishingPort;
+
+	/** The default name to use when publishing face tracking name */
+	UPROPERTY(Config, EditAnywhere, Category="AR Settings")
+	FName DefaultFaceTrackingLiveLinkSubjectName;
+
+	/** The default tracking to use when tracking face blend shapes (face relative or mirrored). Defaults to face relative */
+	UPROPERTY(Config, EditAnywhere, Category="AR Settings")
+	EARFaceTrackingDirection DefaultFaceTrackingDirection;
+
+	/** Whether to adjust thread priorities during an AR session or not */
+	UPROPERTY(Config, EditAnywhere, Category="AR Settings")
+	bool bAdjustThreadPrioritiesDuringARSession;
+	
+	/** The game thread priority to change to when an AR session is running, default is 47 */
+	UPROPERTY(Config, EditAnywhere, Category="AR Settings")
+	int32 GameThreadPriorityOverride;
+	
+	/** The render thread priority to change to when an AR session is running, default is 45 */
+	UPROPERTY(Config, EditAnywhere, Category="AR Settings")
+	int32 RenderThreadPriorityOverride;
+};
diff --git a/Engine/Plugins/Experimental/AR/Apple/AppleARKit/Source/AppleARKit/Public/AppleARKitTextures.h b/Engine/Plugins/Experimental/AR/Apple/AppleARKit/Source/AppleARKit/Public/AppleARKitTextures.h
new file mode 100644
--- /dev/null
+++ b/Engine/Plugins/Experimental/AR/Apple/AppleARKit/Source/AppleARKit/Public/AppleARKitTextures.h
@@ -0,0 +1,121 @@
+// Copyright 1998-2018 Epic Games, Inc. All Rights Reserved.
+
+#pragma once
+
+#include "AppleARKitAvailability.h"
+#include "ARTextures.h"
+#if PLATFORM_MAC || PLATFORM_IOS
+	#import <CoreVideo/CoreVideo.h>
+	#import <AVFoundation/AVFoundation.h>
+#endif
+#include "AppleImageUtilsTypes.h"
+#include "AppleARKitTextures.generated.h"
+
+UCLASS(BlueprintType)
+class APPLEARKIT_API UAppleARKitTextureCameraImage :
+	public UARTextureCameraImage,
+	public IAppleImageInterface
+{
+	GENERATED_UCLASS_BODY()
+
+public:
+	// UTexture interface implementation
+	virtual void BeginDestroy() override;
+	virtual FTextureResource* CreateResource() override;
+	virtual EMaterialValueType GetMaterialType() const override { return MCT_TextureExternal; }
+	virtual float GetSurfaceWidth() const override { return Size.X; }
+	virtual float GetSurfaceHeight() const override { return Size.Y; }
+	virtual FGuid GetExternalTextureGuid() const override { return ExternalTextureGuid; }
+	// End UTexture interface
+
+#if SUPPORTS_ARKIT_1_0
+	/** Sets any initialization data */
+	virtual void Init(float InTimestamp, CVPixelBufferRef InCameraImage);
+#endif
+
+	virtual EAppleTextureType GetTextureType() const override { return EAppleTextureType::PixelBuffer; }
+
+#if PLATFORM_MAC || PLATFORM_IOS
+	/** Returns the cached camera image. You must retain this if you hold onto it */
+	CVPixelBufferRef GetCameraImage() { return CameraImage; }
+
+	// IAppleImageInterface interface implementation
+	virtual CVPixelBufferRef GetPixelBuffer() const override { return CameraImage; }
+	// End IAppleImageInterface interface
+#endif
+
+private:
+#if PLATFORM_MAC || PLATFORM_IOS
+	/** The Apple specific representation of the ar camera image */
+	CVPixelBufferRef CameraImage;
+#endif
+};
+
+UCLASS(BlueprintType)
+class APPLEARKIT_API UAppleARKitTextureCameraDepth :
+	public UARTextureCameraDepth
+{
+	GENERATED_UCLASS_BODY()
+
+public:
+	// UTexture interface implementation
+	virtual void BeginDestroy() override;
+	virtual FTextureResource* CreateResource() override;
+	virtual EMaterialValueType GetMaterialType() const override { return MCT_TextureExternal; }
+	virtual float GetSurfaceWidth() const override { return Size.X; }
+	virtual float GetSurfaceHeight() const override { return Size.Y; }
+	virtual FGuid GetExternalTextureGuid() const override { return ExternalTextureGuid; }
+	// End UTexture interface
+
+#if SUPPORTS_ARKIT_1_0
+	/** Sets any initialization data */
+	virtual void Init(float InTimestamp, AVDepthData* InCameraDepth);
+#endif
+
+#if PLATFORM_MAC || PLATFORM_IOS
+	/** Returns the cached camera depth. You must retain this if you hold onto it */
+	AVDepthData* GetCameraDepth() { return CameraDepth; }
+#endif
+
+private:
+#if PLATFORM_MAC || PLATFORM_IOS
+	/** The Apple specific representation of the ar depth image */
+	AVDepthData* CameraDepth;
+#endif
+};
+
+UCLASS(BlueprintType)
+class APPLEARKIT_API UAppleARKitEnvironmentCaptureProbeTexture :
+	public UAREnvironmentCaptureProbeTexture,
+	public IAppleImageInterface
+{
+	GENERATED_UCLASS_BODY()
+	
+public:
+	
+	// UTexture interface implementation
+	virtual void BeginDestroy() override;
+	virtual FTextureResource* CreateResource() override;
+	virtual EMaterialValueType GetMaterialType() const override { return MCT_TextureExternal; }
+	virtual float GetSurfaceWidth() const override { return Size.X; }
+	virtual float GetSurfaceHeight() const override { return Size.Y; }
+	virtual FGuid GetExternalTextureGuid() const override { return ExternalTextureGuid; }
+	// End UTexture interface
+	
+	virtual EAppleTextureType GetTextureType() const override { return EAppleTextureType::MetalTexture; }
+
+#if PLATFORM_MAC || PLATFORM_IOS
+	/** Sets any initialization data */
+	virtual void Init(float InTimestamp, id<MTLTexture> InEnvironmentTexture);
+
+	// IAppleImageInterface interface implementation
+	virtual id<MTLTexture> GetMetalTexture() const override { return MetalTexture; }
+	// End IAppleImageInterface interface
+
+private:
+	/** The Apple specific representation of the ar environment texture */
+	id<MTLTexture> MetalTexture;
+#endif
+};
+
+
diff --git a/Engine/Plugins/Experimental/AR/Google/GoogleARCore/Source/GoogleARCoreBase/GoogleARCoreBase.Build.cs b/Engine/Plugins/Experimental/AR/Google/GoogleARCore/Source/GoogleARCoreBase/GoogleARCoreBase.Build.cs
new file mode 100644
--- /dev/null
+++ b/Engine/Plugins/Experimental/AR/Google/GoogleARCore/Source/GoogleARCoreBase/GoogleARCoreBase.Build.cs
@@ -0,0 +1,89 @@
+// Copyright 2017 Google Inc.
+
+
+using System.IO;
+
+
+namespace UnrealBuildTool.Rules
+{
+	public class GoogleARCoreBase : ModuleRules
+	{
+		public GoogleARCoreBase(ReadOnlyTargetRules Target) : base(Target)
+		{
+			PrivateIncludePaths.AddRange(
+				new string[]
+				{
+					"GoogleARCoreBase/Private",
+				}
+			);
+			PublicDependencyModuleNames.AddRange(
+					new string[]
+					{
+						"HeadMountedDisplay",
+						"AugmentedReality",
+					}
+				);
+
+			PrivateDependencyModuleNames.AddRange(
+				new string[]
+				{
+					"Core",
+					"CoreUObject",
+					"Engine",
+					"EngineSettings",
+					"Slate",
+					"SlateCore",
+					"RHI",
+					"RenderCore",
+					"ShaderCore",
+					"AndroidPermission",
+					"GoogleARCoreRendering",
+					"GoogleARCoreSDK",
+					"OpenGL",
+					"UElibPNG",
+					"zlib"
+				}
+			);
+
+			PrivateIncludePathModuleNames.AddRange(
+				new string[]
+				{
+					"Settings" // For editor settings panel.
+				}
+			);
+
+			if (Target.Platform == UnrealTargetPlatform.Android)
+			{
+				// Additional dependencies on android...
+				PrivateDependencyModuleNames.Add("Launch");
+
+				// Register Plugin Language
+				AdditionalPropertiesForReceipt.Add("AndroidPlugin", Path.Combine(ModuleDirectory, "GoogleARCoreBase_APL.xml"));
+			}
+
+			if (Target.bBuildEditor)
+			{
+				string ExecName = "";
+				if (Target.Platform == UnrealTargetPlatform.Win64)
+				{
+					ExecName = "Windows/arcoreimg.exe";
+				}
+				else if (Target.Platform == UnrealTargetPlatform.Linux)
+				{
+					ExecName = "Linux/arcoreimg";
+				}
+				else if (Target.Platform == UnrealTargetPlatform.Mac)
+				{
+					ExecName = "Mac/ptdbtool_macos_lipobin";
+				}
+				
+				if (ExecName.Length > 0)
+				{
+					RuntimeDependencies.Add("$(EngineDir)/Plugins/Runtime/GoogleARCore/Binaries/ThirdParty/Google/ARCoreImg/" + ExecName);
+				}
+			}
+			
+			bFasterWithoutUnity = false;
+        }
+    }
+}
diff --git a/Engine/Plugins/Experimental/AR/Google/GoogleARCore/Source/GoogleARCoreBase/Private/GoogleARCoreAPI.h b/Engine/Plugins/Experimental/AR/Google/GoogleARCore/Source/GoogleARCoreBase/Private/GoogleARCoreAPI.h
new file mode 100644
--- /dev/null
+++ b/Engine/Plugins/Experimental/AR/Google/GoogleARCore/Source/GoogleARCoreBase/Private/GoogleARCoreAPI.h
@@ -0,0 +1,510 @@
+// Copyright 2017 Google Inc.
+#pragma once
+
+#include "CoreMinimal.h"
+#include "UObject/Object.h"
+
+#include "GoogleARCoreTypes.h"
+#include "GoogleARCoreSessionConfig.h"
+#include "GoogleARCoreCameraImageBlitter.h"
+#include "GoogleARCoreAugmentedImage.h"
+#include "GoogleARCoreCameraIntrinsics.h"
+#include "ARSessionConfig.h"
+
+#if PLATFORM_ANDROID
+#include "arcore_c_api.h"
+#endif
+
+#include "GoogleARCoreAPI.generated.h"
+
+DEFINE_LOG_CATEGORY_STATIC(LogGoogleARCoreAPI, Log, All);
+
+enum class EGoogleARCoreAPIStatus : int
+{
+	/// The operation was successful.
+	AR_SUCCESS = 0,
+
+	/// One of the arguments was invalid, either null or not appropriate for the
+	/// operation requested.
+	AR_ERROR_INVALID_ARGUMENT = -1,
+
+	/// An internal error occurred that the application should not attempt to
+	/// recover from.
+	AR_ERROR_FATAL = -2,
+
+	/// An operation was attempted that requires the session be running, but the
+	/// session was paused.
+	AR_ERROR_SESSION_PAUSED = -3,
+
+	/// An operation was attempted that requires the session be paused, but the
+	/// session was running.
+	AR_ERROR_SESSION_NOT_PAUSED = -4,
+
+	/// An operation was attempted that the session be in the TRACKING state,
+	/// but the session was not.
+	AR_ERROR_NOT_TRACKING = -5,
+
+	/// A texture name was not set by calling ArSession_setCameraTextureName()
+	/// before the first call to ArSession_update()
+	AR_ERROR_TEXTURE_NOT_SET = -6,
+
+	/// An operation required GL context but one was not available.
+	AR_ERROR_MISSING_GL_CONTEXT = -7,
+
+	/// The configuration supplied to ArSession_configure() was unsupported.
+	/// To avoid this error, ensure that Session_checkSupported() returns true.
+	AR_ERROR_UNSUPPORTED_CONFIGURATION = -8,
+
+	/// The android camera permission has not been granted prior to calling
+	/// ArSession_resume()
+	AR_ERROR_CAMERA_PERMISSION_NOT_GRANTED = -9,
+
+	/// Acquire failed because the object being acquired is already released.
+	/// For example, this happens if the application holds an ::ArFrame beyond
+	/// the next call to ArSession_update(), and then tries to acquire its point
+	/// cloud.
+	AR_ERROR_DEADLINE_EXCEEDED = -10,
+
+	/// There are no available resources to complete the operation.  In cases of
+	/// @c acquire methods returning this error, This can be avoided by
+	/// releasing previously acquired objects before acquiring new ones.
+	AR_ERROR_RESOURCE_EXHAUSTED = -11,
+
+	/// Acquire failed because the data isn't available yet for the current
+	/// frame. For example, acquire the image metadata may fail with this error
+	/// because the camera hasn't fully started.
+	AR_ERROR_NOT_YET_AVAILABLE = -12,
+
+	/// The android camera has been reallocated to a higher priority app or is
+	/// otherwise unavailable.
+	AR_ERROR_CAMERA_NOT_AVAILABLE = -13,
+
+	/// The data passed in for this operation was not in a valid format.
+	AR_ERROR_DATA_INVALID_FORMAT = -18,
+
+	/// The ARCore APK is not installed on this device.
+	AR_UNAVAILABLE_ARCORE_NOT_INSTALLED = -100,
+
+	/// The device is not currently compatible with ARCore.
+	AR_UNAVAILABLE_DEVICE_NOT_COMPATIBLE = -101,
+
+	/// The ARCore APK currently installed on device is too old and needs to be
+	/// updated.
+	AR_UNAVAILABLE_APK_TOO_OLD = -103,
+
+	/// The ARCore APK currently installed no longer supports the ARCore SDK
+	/// that the application was built with.
+	AR_UNAVAILABLE_SDK_TOO_OLD = -104,
+
+	/// The user declined installation of the ARCore APK during this run of the
+	/// application and the current request was not marked as user-initiated.
+	AR_UNAVAILABLE_USER_DECLINED_INSTALLATION = -105
+};
+
+#if PLATFORM_ANDROID
+static ArTrackableType GetTrackableType(UClass* ClassType)
+{
+	if (ClassType == UARTrackedGeometry::StaticClass())
+	{
+		return ArTrackableType::AR_TRACKABLE_BASE_TRACKABLE;
+	}
+	else if (ClassType == UARPlaneGeometry::StaticClass())
+	{
+		return ArTrackableType::AR_TRACKABLE_PLANE;
+	}
+	else if (ClassType == UARTrackedPoint::StaticClass())
+	{
+		return ArTrackableType::AR_TRACKABLE_POINT;
+	}
+	else if (ClassType == UGoogleARCoreAugmentedImage::StaticClass())
+	{
+		return ArTrackableType::AR_TRACKABLE_AUGMENTED_IMAGE;
+	}
+	else
+	{
+		return ArTrackableType::AR_TRACKABLE_NOT_VALID;
+	}
+}
+#endif
+
+class FGoogleARCoreFrame;
+class FGoogleARCoreSession;
+class UGoogleARCoreCameraImage;
+
+UCLASS()
+class UGoogleARCoreUObjectManager : public UObject
+{
+	GENERATED_BODY()
+
+public:
+	UPROPERTY()
+	TArray<UARPin*> AllAnchors;
+
+	UPROPERTY()
+	UGoogleARCorePointCloud* LatestPointCloud;
+
+#if PLATFORM_ANDROID
+	TMap<ArAnchor*, UARPin*> HandleToAnchorMap;
+	TMap<ArTrackable*, TWeakObjectPtr<UARTrackedGeometry>> TrackableHandleMap;
+
+	template< class T > T* GetTrackableFromHandle(ArTrackable* TrackableHandle, FGoogleARCoreSession* Session);
+
+	void DumpTrackableHandleMap(const ArSession* SessionHandle);
+#endif
+};
+
+class FGoogleARCoreAPKManager
+{
+public:
+	static EGoogleARCoreAvailability CheckARCoreAPKAvailability();
+	static EGoogleARCoreAPIStatus RequestInstall(bool bUserRequestedInstall, EGoogleARCoreInstallStatus& OutInstallStatus);
+};
+
+
+class FGoogleARCoreSession : public TSharedFromThis<FGoogleARCoreSession>, public FGCObject
+{
+
+public:
+	static TSharedPtr<FGoogleARCoreSession> CreateARCoreSession();
+
+	FGoogleARCoreSession();
+	~FGoogleARCoreSession();
+
+	// Properties
+	EGoogleARCoreAPIStatus GetSessionCreateStatus();
+	UGoogleARCoreUObjectManager* GetUObjectManager();
+	float GetWorldToMeterScale();
+	void SetARSystem(TSharedRef<FARSupportInterface , ESPMode::ThreadSafe> InArSystem) { ARSystem = InArSystem; }
+	TSharedRef<FARSupportInterface , ESPMode::ThreadSafe> GetARSystem() { return ARSystem.ToSharedRef(); }
+#if PLATFORM_ANDROID
+	ArSession* GetHandle();
+#endif
+
+	// Lifecycle
+	bool IsConfigSupported(const UARSessionConfig& Config);
+	EGoogleARCoreAPIStatus ConfigSession(const UARSessionConfig& Config);
+	TArray<FGoogleARCoreCameraConfig> GetSupportedCameraConfig();
+	EGoogleARCoreAPIStatus SetCameraConfig(FGoogleARCoreCameraConfig CameraConfig);
+	void GetARCameraConfig(FGoogleARCoreCameraConfig& OutCurrentCameraConfig);
+	EGoogleARCoreAPIStatus Resume();
+	EGoogleARCoreAPIStatus Pause();
+	EGoogleARCoreAPIStatus Update(float WorldToMeterScale);
+	const FGoogleARCoreFrame* GetLatestFrame();
+	uint32 GetFrameNum();
+
+	void SetCameraTextureId(uint32_t TextureId);
+	void SetDisplayGeometry(int Rotation, int Width, int Height);
+
+	// Anchor API
+	EGoogleARCoreAPIStatus CreateARAnchor(const FTransform& TransfromInTrackingSpace, UARTrackedGeometry* TrackedGeometry, USceneComponent* ComponentToPin, FName InDebugName, UARPin*& OutAnchor);
+	void DetachAnchor(UARPin* Anchor);
+
+	void GetAllAnchors(TArray<UARPin*>& OutAnchors) const;
+	template< class T > void GetAllTrackables(TArray<T*>& OutARCoreTrackableList);
+
+	void* GetLatestFrameRawPointer();
+
+private:
+	virtual void AddReferencedObjects(FReferenceCollector& Collector) override;
+	EGoogleARCoreAPIStatus SessionCreateStatus;
+	const UARSessionConfig* SessionConfig;
+	FGoogleARCoreFrame* LatestFrame;
+	UGoogleARCoreUObjectManager* UObjectManager;
+	uint32_t CameraTextureId;
+	float CachedWorldToMeterScale;
+	uint32 FrameNumber;
+
+	TSharedPtr<FARSupportInterface , ESPMode::ThreadSafe> ARSystem;
+
+#if PLATFORM_ANDROID
+	ArSession* SessionHandle = nullptr;
+	ArConfig* ConfigHandle = nullptr;
+#endif
+};
+
+class FGoogleARCoreFrame
+{
+	friend class FGoogleARCoreSession;
+
+public:
+	FGoogleARCoreFrame(FGoogleARCoreSession* Session);
+	~FGoogleARCoreFrame();
+
+	void Init();
+
+	void Update(float WorldToMeterScale);
+
+	FTransform GetCameraPose() const;
+	int64 GetCameraTimestamp() const;
+	EGoogleARCoreTrackingState GetCameraTrackingState() const;
+
+	void GetUpdatedAnchors(TArray<UARPin*>& OutUpdatedAnchors) const;
+	template< class T > void GetUpdatedTrackables(TArray<T*>& OutARCoreTrackableList) const;
+
+	void ARLineTrace(const FVector2D& ScreenPosition, EGoogleARCoreLineTraceChannel RequestedTraceChannels, TArray<FARTraceResult>& OutHitResults) const;
+	void ARLineTrace(const FVector& Start, const FVector& End, EGoogleARCoreLineTraceChannel RequestedTraceChannels, TArray<FARTraceResult>& OutHitResults) const;
+
+	bool IsDisplayRotationChanged() const;
+	FMatrix GetProjectionMatrix() const;
+	void TransformDisplayUvCoords(const TArray<float>& UvCoords, TArray<float>& OutUvCoords) const;
+
+	FGoogleARCoreLightEstimate GetLightEstimate() const;
+	EGoogleARCoreAPIStatus GetPointCloud(UGoogleARCorePointCloud*& OutLatestPointCloud) const;
+	EGoogleARCoreAPIStatus AcquirePointCloud(UGoogleARCorePointCloud*& OutLatestPointCloud) const;
+	EGoogleARCoreAPIStatus AcquireCameraImage(UGoogleARCoreCameraImage *&OutCameraImage) const;
+	EGoogleARCoreAPIStatus GetCameraImageIntrinsics(
+		UGoogleARCoreCameraIntrinsics *&OutCameraIntrinsics) const;
+	EGoogleARCoreAPIStatus GetCameraTextureIntrinsics(
+		UGoogleARCoreCameraIntrinsics *&OutCameraIntrinsics) const;
+
+#if PLATFORM_ANDROID
+	EGoogleARCoreAPIStatus GetCameraMetadata(const ACameraMetadata*& OutCameraMetadata) const;
+	ArFrame* GetHandle() { return FrameHandle; };
+#endif
+
+private:
+	FGoogleARCoreSession* Session;
+	FTransform LatestCameraPose;
+	int64 LatestCameraTimestamp;
+	EGoogleARCoreTrackingState LatestCameraTrackingState;
+
+	EGoogleARCoreAPIStatus LatestPointCloudStatus;
+	EGoogleARCoreAPIStatus LatestImageMetadataStatus;
+
+	TArray<UARPin*> UpdatedAnchors;
+
+#if PLATFORM_ANDROID
+	const ArSession* SessionHandle = nullptr;
+	ArFrame* FrameHandle = nullptr;
+	ArCamera* CameraHandle = nullptr;
+	ArPose* SketchPoseHandle = nullptr;
+	ArImageMetadata* LatestImageMetadata = nullptr;
+
+	void FilterLineTraceResults(ArHitResultList* HitResultList, EGoogleARCoreLineTraceChannel RequestedTraceChannels, 
+		TArray<FARTraceResult>& OutHitResults, float MaxDistance = TNumericLimits<float>::Max()) const;
+#endif
+};
+
+class FGoogleARCoreTrackableResource : public IARRef
+{
+public:
+	// IARRef interface
+	virtual void AddRef() override { }
+
+	virtual void RemoveRef() override
+	{
+#if PLATFORM_ANDROID
+		ArTrackable_release(TrackableHandle);
+		TrackableHandle = nullptr;
+#endif
+	}
+
+#if PLATFORM_ANDROID
+public:
+	FGoogleARCoreTrackableResource(TSharedPtr<FGoogleARCoreSession> InSession, ArTrackable* InTrackableHandle, UARTrackedGeometry* InTrackedGeometry)
+		: Session(InSession)
+		, TrackableHandle(InTrackableHandle)
+		, TrackedGeometry(InTrackedGeometry)
+	{
+		ensure(TrackableHandle != nullptr);
+	}
+
+	virtual ~FGoogleARCoreTrackableResource()
+	{
+		ArTrackable_release(TrackableHandle);
+		TrackableHandle = nullptr;
+	}
+
+	EARTrackingState GetTrackingState();
+
+	virtual void UpdateGeometryData();
+
+	TWeakPtr<FGoogleARCoreSession> GetSession() { return Session; }
+	ArTrackable* GetNativeHandle() { return TrackableHandle; }
+
+	void ResetNativeHandle(ArTrackable* InTrackableHandle);
+
+protected:
+	TWeakPtr<FGoogleARCoreSession> Session;
+	ArTrackable* TrackableHandle;
+	UARTrackedGeometry* TrackedGeometry;
+#endif
+};
+
+class FGoogleARCoreTrackedPlaneResource : public FGoogleARCoreTrackableResource
+{
+public:
+#if PLATFORM_ANDROID
+	FGoogleARCoreTrackedPlaneResource(TSharedPtr<FGoogleARCoreSession> InSession, ArTrackable* InTrackableHandle, UARTrackedGeometry* InTrackedGeometry)
+		: FGoogleARCoreTrackableResource(InSession, InTrackableHandle, InTrackedGeometry)
+	{
+		ensure(TrackableHandle != nullptr);
+	}
+
+	void UpdateGeometryData() override;
+
+	ArPlane* GetPlaneHandle() { return reinterpret_cast<ArPlane*>(TrackableHandle); }
+#endif
+};
+
+class FGoogleARCoreTrackedPointResource : public FGoogleARCoreTrackableResource
+{
+public:
+#if PLATFORM_ANDROID
+	FGoogleARCoreTrackedPointResource(TSharedPtr<FGoogleARCoreSession> InSession, ArTrackable* InTrackableHandle, UARTrackedGeometry* InTrackedGeometry)
+		: FGoogleARCoreTrackableResource(InSession, InTrackableHandle, InTrackedGeometry)
+	{
+		ensure(TrackableHandle != nullptr);
+	}
+
+	void UpdateGeometryData() override;
+
+	ArPoint* GetPointHandle() { return reinterpret_cast<ArPoint*>(TrackableHandle); }
+#endif
+};
+
+class FGoogleARCoreAugmentedImageResource : public FGoogleARCoreTrackableResource
+{
+public:
+#if PLATFORM_ANDROID
+	FGoogleARCoreAugmentedImageResource(TSharedPtr<FGoogleARCoreSession> InSession, ArTrackable* InTrackableHandle, UARTrackedGeometry* InTrackedGeometry)
+		: FGoogleARCoreTrackableResource(InSession, InTrackableHandle, InTrackedGeometry)
+	{
+		ensure(TrackableHandle != nullptr);
+	}
+
+	void UpdateGeometryData() override;
+
+	ArAugmentedImage* GetImageHandle() { return reinterpret_cast<ArAugmentedImage*>(TrackableHandle); }
+#endif
+};
+
+#if PLATFORM_ANDROID
+// Template function definition
+template< class T >
+T* UGoogleARCoreUObjectManager::GetTrackableFromHandle(ArTrackable* TrackableHandle, FGoogleARCoreSession* Session)
+{
+	if (!TrackableHandleMap.Contains(TrackableHandle)
+		|| !TrackableHandleMap[TrackableHandle].IsValid()
+		|| TrackableHandleMap[TrackableHandle]->GetTrackingState() == EARTrackingState::StoppedTracking)
+	{
+		// Add the trackable to the cache.
+		UARTrackedGeometry* NewTrackableObject = nullptr;
+		ArTrackableType TrackableType = ArTrackableType::AR_TRACKABLE_NOT_VALID;
+		ArTrackable_getType(Session->GetHandle(), TrackableHandle, &TrackableType);
+		IARRef* NativeResource = nullptr;
+		if (TrackableType == ArTrackableType::AR_TRACKABLE_PLANE)
+		{
+			UARPlaneGeometry* PlaneObject = NewObject<UARPlaneGeometry>();
+			NewTrackableObject = static_cast<UARTrackedGeometry*>(PlaneObject);
+			NativeResource = new FGoogleARCoreTrackedPlaneResource(Session->AsShared(), TrackableHandle, NewTrackableObject);
+		}
+		else if (TrackableType == ArTrackableType::AR_TRACKABLE_POINT)
+		{
+			UARTrackedPoint* PointObject = NewObject<UARTrackedPoint>();
+			NewTrackableObject = static_cast<UARTrackedGeometry*>(PointObject);
+			NativeResource = new FGoogleARCoreTrackedPointResource(Session->AsShared(), TrackableHandle, NewTrackableObject);
+		}
+		else if (TrackableType == ArTrackableType::AR_TRACKABLE_AUGMENTED_IMAGE)
+		{
+			UGoogleARCoreAugmentedImage* ImageObject = NewObject<UGoogleARCoreAugmentedImage>();
+			NewTrackableObject = static_cast<UARTrackedGeometry*>(ImageObject);
+			NativeResource = new FGoogleARCoreAugmentedImageResource(Session->AsShared(), TrackableHandle, NewTrackableObject);
+		}
+
+		// We should have a valid trackable object now.
+		checkf(NewTrackableObject, TEXT("Unknown ARCore Trackable Type: %d"), TrackableType);
+
+		NewTrackableObject->InitializeNativeResource(NativeResource);
+		NativeResource = nullptr;
+
+		FGoogleARCoreTrackableResource* TrackableResource = reinterpret_cast<FGoogleARCoreTrackableResource*>(NewTrackableObject->GetNativeResource());
+
+		// Update the tracked geometry data using the native resource
+		TrackableResource->UpdateGeometryData();
+
+		TrackableHandleMap.Add(TrackableHandle, TWeakObjectPtr<UARTrackedGeometry>(NewTrackableObject));
+	}
+	else
+	{
+		// If we are not create new trackable object, release the trackable handle.
+		ArTrackable_release(TrackableHandle);
+	}
+
+	T* Result = Cast<T>(TrackableHandleMap[TrackableHandle].Get());
+	checkf(Result, TEXT("UGoogleARCoreUObjectManager failed to get a valid trackable %p from the map."), TrackableHandle);
+	return Result;
+}
+#endif
+
+template< class T >
+void FGoogleARCoreFrame::GetUpdatedTrackables(TArray<T*>& OutARCoreTrackableList) const
+{
+	OutARCoreTrackableList.Empty();
+#if PLATFORM_ANDROID
+	if (SessionHandle == nullptr)
+	{
+		return;
+	}
+
+	ArTrackableType TrackableType = GetTrackableType(T::StaticClass());
+	if (TrackableType == ArTrackableType::AR_TRACKABLE_NOT_VALID)
+	{
+		UE_LOG(LogGoogleARCoreAPI, Error, TEXT("Invalid Trackable type: %s"), *T::StaticClass()->GetName());
+		return;
+	}
+
+	ArTrackableList* TrackableListHandle = nullptr;
+	ArTrackableList_create(Session->GetHandle(), &TrackableListHandle);
+	ArFrame_getUpdatedTrackables(Session->GetHandle(), FrameHandle, TrackableType, TrackableListHandle);
+
+	int TrackableListSize = 0;
+	ArTrackableList_getSize(Session->GetHandle(), TrackableListHandle, &TrackableListSize);
+
+	for (int i = 0; i < TrackableListSize; i++)
+	{
+		ArTrackable* TrackableHandle = nullptr;
+		ArTrackableList_acquireItem(Session->GetHandle(), TrackableListHandle, i, &TrackableHandle);
+		T* TrackableObject = Session->GetUObjectManager()->template GetTrackableFromHandle<T>(TrackableHandle, Session);
+
+		OutARCoreTrackableList.Add(TrackableObject);
+	}
+	ArTrackableList_destroy(TrackableListHandle);
+#endif
+}
+
+template< class T >
+void FGoogleARCoreSession::GetAllTrackables(TArray<T*>& OutARCoreTrackableList)
+{
+	OutARCoreTrackableList.Empty();
+#if PLATFORM_ANDROID
+	if (SessionHandle == nullptr)
+	{
+		return;
+	}
+
+	ArTrackableType TrackableType = GetTrackableType(T::StaticClass());
+	if (TrackableType == ArTrackableType::AR_TRACKABLE_NOT_VALID)
+	{
+		UE_LOG(LogGoogleARCoreAPI, Error, TEXT("Invalid Trackable type: %s"), *T::StaticClass()->GetName());
+		return;
+	}
+
+	ArTrackableList* TrackableListHandle = nullptr;
+	ArTrackableList_create(SessionHandle, &TrackableListHandle);
+	ArSession_getAllTrackables(SessionHandle, TrackableType, TrackableListHandle);
+
+	int TrackableListSize = 0;
+	ArTrackableList_getSize(SessionHandle, TrackableListHandle, &TrackableListSize);
+
+	for (int i = 0; i < TrackableListSize; i++)
+	{
+		ArTrackable* TrackableHandle = nullptr;
+		ArTrackableList_acquireItem(SessionHandle, TrackableListHandle, i, &TrackableHandle);
+
+		T* TrackableObject = UObjectManager->template GetTrackableFromHandle<T>(TrackableHandle, this);
+		OutARCoreTrackableList.Add(TrackableObject);
+	}
+	ArTrackableList_destroy(TrackableListHandle);
+#endif
+}
diff --git a/Engine/Plugins/Experimental/AR/Google/GoogleARCore/Source/GoogleARCoreBase/Private/GoogleARCoreAugmentedImage.cpp b/Engine/Plugins/Experimental/AR/Google/GoogleARCore/Source/GoogleARCoreBase/Private/GoogleARCoreAugmentedImage.cpp
new file mode 100644
--- /dev/null
+++ b/Engine/Plugins/Experimental/AR/Google/GoogleARCore/Source/GoogleARCoreBase/Private/GoogleARCoreAugmentedImage.cpp
@@ -0,0 +1,102 @@
+// Copyright 2018 Google Inc.
+
+#include "GoogleARCoreAugmentedImage.h"
+
+#include "GoogleARCoreAugmentedImageDatabase.h"
+#include "GoogleARCoreDevice.h"
+#include "GoogleARCoreAPI.h"
+
+#include "CoreMinimal.h"
+#include "DrawDebugHelpers.h"
+
+#if PLATFORM_ANDROID
+#include "arcore_c_api.h"
+#endif
+
+void UGoogleARCoreAugmentedImage::DebugDraw(
+	UWorld* World, const FLinearColor& OutlineColor,
+	float OutlineThickness, float PersistForSeconds) const
+{
+#if PLATFORM_ANDROID
+
+	FTransform CenterTransform =
+		GetLocalToTrackingTransform();
+
+	FVector HalfExtent(Extent / 2.0f);
+	FVector Corner0(-HalfExtent.X, -HalfExtent.Y, 0.0f);
+	FVector Corner1( HalfExtent.X, -HalfExtent.Y, 0.0f);
+	FVector Corner2( HalfExtent.X,  HalfExtent.Y, 0.0f);
+	FVector Corner3(-HalfExtent.X,  HalfExtent.Y, 0.0f);
+
+	DrawDebugLine(
+		World,
+		CenterTransform.TransformPosition(Corner0),
+		CenterTransform.TransformPosition(Corner1),
+		OutlineColor.ToFColor(false),
+		false,
+		PersistForSeconds,
+		0, OutlineThickness);
+
+	DrawDebugLine(
+		World,
+		CenterTransform.TransformPosition(Corner1),
+		CenterTransform.TransformPosition(Corner2),
+		OutlineColor.ToFColor(false),
+		false,
+		PersistForSeconds,
+		0, OutlineThickness);
+
+	DrawDebugLine(
+		World,
+		CenterTransform.TransformPosition(Corner2),
+		CenterTransform.TransformPosition(Corner3),
+		OutlineColor.ToFColor(false),
+		false,
+		PersistForSeconds,
+		0, OutlineThickness);
+
+	DrawDebugLine(
+		World,
+		CenterTransform.TransformPosition(Corner3),
+		CenterTransform.TransformPosition(Corner0),
+		OutlineColor.ToFColor(false),
+		false,
+		PersistForSeconds,
+		0, OutlineThickness);
+
+#endif
+}
+
+void UGoogleARCoreAugmentedImage::UpdateTrackedGeometry(
+	const TSharedRef<FARSupportInterface , ESPMode::ThreadSafe>& InTrackingSystem,
+	uint32 FrameNumber, double Timestamp, const FTransform& InLocalToTrackingTransform,
+	const FTransform& InAlignmentTransform, const FVector &InCenter, const FVector &InExtent,
+	int32 InImageIndex, const FString& InImageName)
+{
+	Super::UpdateTrackedGeometry(InTrackingSystem, FrameNumber, Timestamp, InLocalToTrackingTransform, InAlignmentTransform);
+
+	Center = InCenter;
+	Extent = InExtent;
+	ImageIndex = InImageIndex;
+	ImageName = InImageName;
+}
+
+FVector UGoogleARCoreAugmentedImage::GetCenter() const
+{
+	return Center;
+}
+
+FVector UGoogleARCoreAugmentedImage::GetExtent() const
+{
+	return Extent;
+}
+
+int32 UGoogleARCoreAugmentedImage::GetImageIndex() const
+{
+	return ImageIndex;
+}
+
+FString UGoogleARCoreAugmentedImage::GetImageName() const
+{
+	return ImageName;
+}
diff --git a/Engine/Plugins/Experimental/AR/Google/GoogleARCore/Source/GoogleARCoreBase/Private/GoogleARCoreAugmentedImageDatabase.cpp b/Engine/Plugins/Experimental/AR/Google/GoogleARCore/Source/GoogleARCoreBase/Private/GoogleARCoreAugmentedImageDatabase.cpp
new file mode 100644
--- /dev/null
+++ b/Engine/Plugins/Experimental/AR/Google/GoogleARCore/Source/GoogleARCoreBase/Private/GoogleARCoreAugmentedImageDatabase.cpp
@@ -0,0 +1,281 @@
+// Copyright 2018 Google Inc.
+
+#include "GoogleARCoreAugmentedImageDatabase.h"
+
+#include "GoogleARCoreBaseLogCategory.h"
+#include "GoogleARCoreDevice.h"
+#include "GoogleARCoreAPI.h"
+
+#include "Kismet/GameplayStatics.h"
+#include "Misc/FileHelper.h"
+#include "HAL/PlatformFilemanager.h"
+#include "HAL/FileManager.h"
+
+#include "png.h"
+
+#if WITH_EDITORONLY_DATA
+
+// Disable warning "interaction between '_setjmp' and C++ object destruction is non-portable"
+// and 'fopen': This function or variable may be unsafe. 
+#ifdef _MSC_VER
+#pragma warning(push)
+#pragma warning(disable:4611)
+#pragma warning(disable:4996)
+#endif
+
+namespace
+{
+	bool GoogleARCoreSaveTextureToPNG(UTexture2D *Tex, const FString &Filename)
+	{
+		TArray<uint8> MipData;
+		bool Ret = true;
+
+		if (Tex->Source.GetMipData(MipData, 0))
+		{
+			if (Tex->Source.GetFormat() != TSF_BGRA8 && Tex->Source.GetFormat() != TSF_RGBA8)
+			{
+				UE_LOG(
+					LogGoogleARCoreAPI, Error,
+					TEXT("Texture %s is not RGBA8 or BGRA8 and cannot be used as a tracking target."),
+					*Tex->GetName());
+
+				return false;
+			}
+
+			int32 Width = Tex->Source.GetSizeX();
+			int32 Height = Tex->Source.GetSizeY();
+
+			png_structp PngPtr = png_create_write_struct(PNG_LIBPNG_VER_STRING, nullptr, nullptr, nullptr);
+			if (PngPtr)
+			{
+				png_infop PngInfoPtr = png_create_info_struct(PngPtr);
+
+				// We're using C file IO here just because of libPNG interop.
+				FILE *OutputFile = fopen(TCHAR_TO_ANSI(*Filename), "wb+");
+
+				if (setjmp(png_jmpbuf(PngPtr)))
+				{
+					UE_LOG(
+						LogGoogleARCoreAPI, Error,
+						TEXT("Error writing PNG for texture %s."),
+						*Tex->GetName());
+					Ret = false;
+				}
+				else
+				{
+					png_init_io(PngPtr, OutputFile);
+					png_set_IHDR(
+						PngPtr, PngInfoPtr, Width, Height,
+						8, PNG_COLOR_TYPE_RGB, PNG_INTERLACE_NONE,
+						PNG_COMPRESSION_TYPE_BASE, PNG_FILTER_TYPE_BASE);
+					png_write_info(PngPtr, PngInfoPtr);
+
+					uint8_t *Row = new uint8_t[Width * 3];
+
+					for (int32 y = 0; y < Height; y++)
+					{
+						for (int32 x = 0; x < Width; x++)
+						{
+							for (int32 c = 0; c < 3; c++)
+							{
+								int32 RealChannel = c;
+								if (Tex->Source.GetFormat() == TSF_BGRA8)
+								{
+									if (RealChannel == 0) {
+										RealChannel = 2;
+									}
+									else if (RealChannel == 2) {
+										RealChannel = 0;
+									}
+								}
+								Row[x * 3 + c] = MipData[(y * Width + x) * 4 + RealChannel];
+							}
+						}
+						png_write_row(PngPtr, Row);
+					}
+
+					png_write_end(PngPtr, NULL);
+
+					delete[] Row;
+				}
+
+				if (OutputFile)
+				{
+					fclose(OutputFile);
+				}
+
+				if (PngInfoPtr)
+				{
+					png_free_data(PngPtr, PngInfoPtr, PNG_FREE_ALL, -1);
+				}
+
+				png_destroy_write_struct(&PngPtr, nullptr);
+			}
+		}
+		else
+		{
+			UE_LOG(
+				LogGoogleARCoreAPI, Error,
+				TEXT("Error reading mip data in texture %s."),
+				*Tex->GetName());
+
+			return false;
+		}
+
+		return Ret;
+	}
+}
+
+// Renable warning "interaction between '_setjmp' and C++ object destruction is non-portable"
+// and 'fopen': This function or variable may be unsafe. 
+#ifdef _MSC_VER
+#pragma warning(pop)
+#endif
+
+#endif
+
+void UGoogleARCoreAugmentedImageDatabase::Serialize(FArchive& Ar)
+{
+	FString OutStdout;
+	FString OutStderr;
+	int32 OutReturnCode = 0;
+
+#if !PLATFORM_ANDROID && WITH_EDITORONLY_DATA
+
+	if (!Ar.IsLoading() && Ar.IsCooking()) {
+
+		SerializedDatabase.Empty();
+
+		if (Entries.Num()) {
+
+			FString PathToDbTool =
+				FPaths::Combine(
+					*FPaths::EnginePluginsDir(),
+					TEXT("Runtime"),
+					TEXT("GoogleARCore"),
+					TEXT("Binaries"),
+					TEXT("ThirdParty"),
+					TEXT("Google"),
+					TEXT("ARCoreImg"),
+					*UGameplayStatics::GetPlatformName(),
+#if PLATFORM_LINUX
+					TEXT("arcoreimg")
+#elif PLATFORM_WINDOWS
+					TEXT("arcoreimg.exe")
+#elif PLATFORM_MAC
+					TEXT("ptdbtool_macos_lipobin")
+#endif
+					);
+
+			FString TempDir =
+				FPaths::ConvertRelativePathToFull(
+					FPaths::Combine(
+						*FPaths::EnginePluginsDir(),
+						TEXT("Runtime"),
+						TEXT("GoogleARCore"),
+						TEXT("Intermediate"),
+						TEXT("ARCoreImgTemp")));
+
+			IPlatformFile& PlatformFile = FPlatformFileManager::Get().GetPlatformFile();
+			if (!PlatformFile.DirectoryExists(*TempDir))
+			{
+				PlatformFile.CreateDirectory(*TempDir);
+			}
+
+			TArray<FString> CleanupList;
+
+			FString ImageListFileContents;
+
+			for (int32 i = 0; i < Entries.Num(); i++) {
+				UTexture2D *Tex = Entries[i].ImageAsset;
+
+				if (Tex) {
+
+					FString PNGFilename =
+						FPaths::Combine(TempDir, Tex->GetName() + FString(".png"));
+
+					if (GoogleARCoreSaveTextureToPNG(
+							Tex,
+							PNGFilename))
+					{
+						// "|" is used as a delimeter in the image
+						// list, and there doesn't seem to be any way
+						// to escape them, so they will be replaced
+						// with underscores.
+						FString TmpName =
+							Entries[i].Name.ToString().Replace(TEXT("|"), TEXT("_"));
+
+						ImageListFileContents +=
+							TmpName +
+							FString("|") +
+							PNGFilename;
+
+						CleanupList.Add(PNGFilename);
+
+						if (Entries[i].Width > 0.0f) {
+							ImageListFileContents +=
+								FString("|") +
+								FString::SanitizeFloat(Entries[i].Width);
+						}
+
+						ImageListFileContents += FString("\n");
+					}
+
+				}
+			}
+
+			FString PathToImageList =
+				FPaths::Combine(TempDir, TEXT("image_list.txt"));
+			FString PathToImageDb =
+				FPaths::Combine(TempDir, TEXT("image_list.imgdb"));
+
+			CleanupList.Add(PathToImageList);
+			CleanupList.Add(PathToImageDb);
+
+			FFileHelper::SaveStringToFile(
+				ImageListFileContents,
+				*PathToImageList);
+
+			OutStderr = "";
+			OutStdout = "";
+			OutReturnCode = 0;
+
+			FPlatformProcess::ExecProcess(
+				*PathToDbTool,
+				*(FString("build-db --input_image_list_path=\"")+ PathToImageList +
+				  FString("\" --output_db_path=\"") + PathToImageDb + "\""),
+				&OutReturnCode,
+				&OutStdout,
+				&OutStderr);
+
+			if (OutReturnCode)
+			{
+				Ar.SetError();
+				Ar.ArIsError = 1;
+				Ar.ArIsCriticalError = 1;
+				UE_LOG(LogGoogleARCoreAPI, Error, TEXT("Failed to build augmented image database: %s"), *OutStderr);
+			}
+			else
+			{
+				FFileHelper::LoadFileToArray(SerializedDatabase, *PathToImageDb, 0);
+				UE_LOG(LogGoogleARCoreAPI, Log,
+					   TEXT("Augmented image database created. Size: %d bytes. Tool output: %s"),
+					   SerializedDatabase.Num(), *OutStdout);
+			}
+
+			for (int32 i = 0; i < CleanupList.Num(); i++)
+			{
+				UE_LOG(LogGoogleARCoreAPI, Log, TEXT("Cleaning up: %s"), *CleanupList[i]);
+				PlatformFile.DeleteFile(*CleanupList[i]);
+			}
+
+			PlatformFile.DeleteDirectory(*TempDir);
+		}
+	}
+
+#endif
+
+	// Must happen AFTER database generation, because we rely on the
+	// UPROPERTY serialization to actually save the data.
+	Super::Serialize(Ar);
+}
diff --git a/Engine/Plugins/Experimental/AR/Google/GoogleARCore/Source/GoogleARCoreBase/Private/GoogleARCoreCameraImage.cpp b/Engine/Plugins/Experimental/AR/Google/GoogleARCore/Source/GoogleARCoreBase/Private/GoogleARCoreCameraImage.cpp
new file mode 100644
--- /dev/null
+++ b/Engine/Plugins/Experimental/AR/Google/GoogleARCore/Source/GoogleARCoreBase/Private/GoogleARCoreCameraImage.cpp
@@ -0,0 +1,80 @@
+// Copyright 2017 Google Inc.
+
+#include "GoogleARCoreCameraImage.h"
+
+#include "GoogleARCoreAPI.h"
+
+#if PLATFORM_ANDROID
+#include "Android/AndroidApplication.h"
+#include "Android/AndroidJNI.h"
+#include "Ndk/NdkImageAPI.h"
+#endif
+
+UGoogleARCoreCameraImage::~UGoogleARCoreCameraImage()
+{
+	Release();
+}
+
+void UGoogleARCoreCameraImage::Release()
+{
+#if PLATFORM_ANDROID
+	if (ArImage)
+	{
+		NdkImage = nullptr;
+		ArImage_release(ArImage);
+		ArImage = nullptr;
+	}
+#endif
+}
+
+int32 UGoogleARCoreCameraImage::GetWidth() const
+{
+	int32_t Width = 0;
+#if PLATFORM_ANDROID
+	if (NdkImage)
+	{
+		AImage_getWidth_dynamic(NdkImage, &Width);
+	}
+#endif
+	return Width;
+}
+
+int32 UGoogleARCoreCameraImage::GetHeight() const
+{
+	int32_t Height = 0;
+#if PLATFORM_ANDROID
+	if (NdkImage)
+	{
+		AImage_getHeight_dynamic(NdkImage, &Height);
+	}
+#endif
+	return Height;
+}
+
+
+int32 UGoogleARCoreCameraImage::GetPlaneCount() const
+{
+	int32_t PlaneCount = 0;
+#if PLATFORM_ANDROID
+	if (NdkImage)
+	{
+		AImage_getNumberOfPlanes_dynamic(NdkImage, &PlaneCount);
+	}
+#endif
+	return PlaneCount;
+}
+
+uint8 *UGoogleARCoreCameraImage::GetPlaneData(
+	int32 Plane, int32 &PixelStride,
+	int32 &RowStride, int32 &DataLength)
+{
+	uint8_t *PlaneData = nullptr;
+#if PLATFORM_ANDROID
+	AImage_getPlanePixelStride_dynamic(NdkImage, Plane, &PixelStride);
+	AImage_getPlaneRowStride_dynamic(NdkImage, Plane, &RowStride);
+	AImage_getPlaneData_dynamic(
+		NdkImage, Plane,
+		&PlaneData, &DataLength);
+#endif
+	return PlaneData;
+}
