Decoupled AR object timestamps and frame numbers from the game thread update rate

diff --git a/Engine/Plugins/Runtime/AR/Apple/AppleARKit/Source/AppleARKit/Private/AppleARKitSystem.cpp b/Engine/Plugins/Runtime/AR/Apple/AppleARKit/Source/AppleARKit/Private/AppleARKitSystem.cpp
--- a/Engine/Plugins/Runtime/AR/Apple/AppleARKit/Source/AppleARKit/Private/AppleARKitSystem.cpp
+++ b/Engine/Plugins/Runtime/AR/Apple/AppleARKit/Source/AppleARKit/Private/AppleARKitSystem.cpp
@@ -202,8 +202,6 @@ FAppleARKitSystem::FAppleARKitSystem()
 , LightEstimate(nullptr)
 , CameraImage(nullptr)
 , CameraDepth(nullptr)
-, GameThreadFrameNumber(0)
-, GameThreadTimestamp(0.0)
 , LastTrackedGeometry_DebugId(0)
 , FaceARSupport(nullptr)
 {
@@ -359,10 +357,6 @@ void FAppleARKitSystem::UpdateFrame()
 		GameThreadFrame = LastReceivedFrame;
 		if (GameThreadFrame.IsValid())
 		{
-			// Used to mark the time at which tracked geometry was updated
-			GameThreadFrameNumber++;
-			GameThreadTimestamp = GameThreadFrame->Timestamp;
-
 #if SUPPORTS_ARKIT_1_0
 			if (GameThreadFrame->CameraImage != nullptr)
 			{
@@ -372,7 +366,7 @@ void FAppleARKitSystem::UpdateFrame()
 					CameraImage = NewObject<UAppleARKitTextureCameraImage>();
 				}
 				// Reuse the UObjects because otherwise the time between GCs causes ARKit to be starved of resources
-                CameraImage->Init(GameThreadTimestamp, GameThreadFrame->CameraImage);
+                CameraImage->Init(FPlatformTime::Seconds(), GameThreadFrame->CameraImage);
 			}
 
 			if (GameThreadFrame->CameraDepth != nullptr)
@@ -383,7 +377,7 @@ void FAppleARKitSystem::UpdateFrame()
 					CameraDepth = NewObject<UAppleARKitTextureCameraDepth>();
 				}
 				// Reuse the UObjects because otherwise the time between GCs causes ARKit to be starved of resources
-                CameraDepth->Init(GameThreadTimestamp, GameThreadFrame->CameraDepth);
+                CameraDepth->Init(FPlatformTime::Seconds(), GameThreadFrame->CameraDepth);
 			}
 #endif
 		}
@@ -736,7 +730,7 @@ UARPin* FAppleARKitSystem::OnPinComponent( USceneComponent* ComponentToPin, cons
 			double UpdateTimestamp = FPlatformTime::Seconds();
 			
 			GeometryToPinTo = NewObject<UARTrackedPoint>();
-			GeometryToPinTo->UpdateTrackedGeometry(GetARCompositionComponent().ToSharedRef(), GameThreadFrameNumber, UpdateTimestamp, PinToTrackingTransform, GetARCompositionComponent()->GetAlignmentTransform());
+			GeometryToPinTo->UpdateTrackedGeometry(GetARCompositionComponent().ToSharedRef(), 0, FPlatformTime::Seconds(), PinToTrackingTransform, GetARCompositionComponent()->GetAlignmentTransform());
 		}
 		
 		UARPin* NewPin = NewObject<UARPin>();
@@ -1390,7 +1384,7 @@ void FAppleARKitSystem::SessionDidFailWithError_DelegateThread(const FString& Er
 
 TArray<int32> FAppleARKitAnchorData::FaceIndices;
 
-static TSharedPtr<FAppleARKitAnchorData> MakeAnchorData( ARAnchor* Anchor )
+static TSharedPtr<FAppleARKitAnchorData> MakeAnchorData( ARAnchor* Anchor, double Timestamp, uint32 FrameNumber )
 {
 	TSharedPtr<FAppleARKitAnchorData> NewAnchor;
 	if ([Anchor isKindOfClass:[ARPlaneAnchor class]])
@@ -1461,6 +1455,9 @@ static TSharedPtr<FAppleARKitAnchorData> MakeAnchorData( ARAnchor* Anchor )
 			FAppleARKitConversion::ToFGuid(Anchor.identifier),
 			FAppleARKitConversion::ToFTransform(Anchor.transform));
 	}
+
+	NewAnchor->Timestamp = Timestamp;
+	NewAnchor->FrameNumber = FrameNumber;
 	
 	return NewAnchor;
 }
@@ -1470,11 +1467,10 @@ void FAppleARKitSystem::SessionDidAddAnchors_DelegateThread( NSArray<ARAnchor*>*
 	// If this object is valid, we are running a face session and need that code to process things
 	if (FaceARSupport != nullptr)
 	{
-		double UpdateTimestamp = FPlatformTime::Seconds();
 		const FRotator& AdjustBy = GetARCompositionComponent()->GetSessionConfig().GetWorldAlignment() == EARWorldAlignment::Camera ? DerivedTrackingToUnrealRotation : FRotator::ZeroRotator;
 		const EARFaceTrackingUpdate UpdateSetting = GetARCompositionComponent()->GetSessionConfig().GetFaceTrackingUpdate();
 
-		const TArray<TSharedPtr<FAppleARKitAnchorData>> AnchorList = FaceARSupport->MakeAnchorData(anchors, UpdateTimestamp, GameThreadFrameNumber, AdjustBy, UpdateSetting);
+		const TArray<TSharedPtr<FAppleARKitAnchorData>> AnchorList = FaceARSupport->MakeAnchorData(anchors, AdjustBy, UpdateSetting);
 		for (TSharedPtr<FAppleARKitAnchorData> NewAnchorData : AnchorList)
 		{
 			auto AddAnchorTask = FSimpleDelegateGraphTask::FDelegate::CreateSP(this, &FAppleARKitSystem::SessionDidAddAnchors_Internal, NewAnchorData.ToSharedRef());
@@ -1483,9 +1479,13 @@ void FAppleARKitSystem::SessionDidAddAnchors_DelegateThread( NSArray<ARAnchor*>*
 		return;
 	}
 
+	// Make sure all anchors get the same timestamp and frame number
+	double Timestamp = FPlatformTime::Seconds();
+	uint32 FrameNumber = FARKitFrameCounter::Get().GetFrameNumber();
+
 	for (ARAnchor* anchor in anchors)
 	{
-		TSharedPtr<FAppleARKitAnchorData> NewAnchorData = MakeAnchorData(anchor);
+		TSharedPtr<FAppleARKitAnchorData> NewAnchorData = MakeAnchorData(anchor, Timestamp, FrameNumber);
 		if (ensure(NewAnchorData.IsValid()))
 		{
 			auto AddAnchorTask = FSimpleDelegateGraphTask::FDelegate::CreateSP(this, &FAppleARKitSystem::SessionDidAddAnchors_Internal, NewAnchorData.ToSharedRef());
@@ -1503,7 +1503,7 @@ void FAppleARKitSystem::SessionDidUpdateAnchors_DelegateThread( NSArray<ARAnchor
 		const FRotator& AdjustBy = GetARCompositionComponent()->GetSessionConfig().GetWorldAlignment() == EARWorldAlignment::Camera ? DerivedTrackingToUnrealRotation : FRotator::ZeroRotator;
 		const EARFaceTrackingUpdate UpdateSetting = GetARCompositionComponent()->GetSessionConfig().GetFaceTrackingUpdate();
 
-		const TArray<TSharedPtr<FAppleARKitAnchorData>> AnchorList = FaceARSupport->MakeAnchorData(anchors, UpdateTimestamp, GameThreadFrameNumber, AdjustBy, UpdateSetting);
+		const TArray<TSharedPtr<FAppleARKitAnchorData>> AnchorList = FaceARSupport->MakeAnchorData(anchors, AdjustBy, UpdateSetting);
 		for (TSharedPtr<FAppleARKitAnchorData> NewAnchorData : AnchorList)
 		{
 			auto UpdateAnchorTask = FSimpleDelegateGraphTask::FDelegate::CreateSP(this, &FAppleARKitSystem::SessionDidUpdateAnchors_Internal, NewAnchorData.ToSharedRef());
@@ -1512,9 +1512,13 @@ void FAppleARKitSystem::SessionDidUpdateAnchors_DelegateThread( NSArray<ARAnchor
 		return;
 	}
 
+	// Make sure all anchors get the same timestamp and frame number
+	double Timestamp = FPlatformTime::Seconds();
+	uint32 FrameNumber = FARKitFrameCounter::Get().GetFrameNumber();
+
 	for (ARAnchor* anchor in anchors)
 	{
-		TSharedPtr<FAppleARKitAnchorData> NewAnchorData = MakeAnchorData(anchor);
+		TSharedPtr<FAppleARKitAnchorData> NewAnchorData = MakeAnchorData(anchor, Timestamp, FrameNumber);
 		if (ensure(NewAnchorData.IsValid()))
 		{
 			auto UpdateAnchorTask = FSimpleDelegateGraphTask::FDelegate::CreateSP(this, &FAppleARKitSystem::SessionDidUpdateAnchors_Internal, NewAnchorData.ToSharedRef());
@@ -1552,7 +1556,7 @@ void FAppleARKitSystem::SessionDidAddAnchors_Internal( TSharedRef<FAppleARKitAnc
 	// If this object is valid, we are running a face session and we need to publish LiveLink data on the game thread
 	if (FaceARSupport != nullptr && AnchorData->AnchorType == EAppleAnchorType::FaceAnchor)
 	{
-		FaceARSupport->PublishLiveLinkData(AnchorData, UpdateTimestamp, GameThreadFrameNumber);
+		FaceARSupport->PublishLiveLinkData(AnchorData);
 	}
 
 	FString NewAnchorDebugName;
@@ -1563,14 +1567,14 @@ void FAppleARKitSystem::SessionDidAddAnchors_Internal( TSharedRef<FAppleARKitAnc
 		{
 			NewAnchorDebugName = FString::Printf(TEXT("ANCHOR-%02d"), LastTrackedGeometry_DebugId++);
 			NewGeometry = NewObject<UARTrackedGeometry>();
-			NewGeometry->UpdateTrackedGeometry(ARComponent.ToSharedRef(), GameThreadFrameNumber, UpdateTimestamp, AnchorData->Transform, GetARCompositionComponent()->GetAlignmentTransform());
+			NewGeometry->UpdateTrackedGeometry(ARComponent.ToSharedRef(), AnchorData->FrameNumber, AnchorData->Timestamp, AnchorData->Transform, GetARCompositionComponent()->GetAlignmentTransform());
 			break;
 		}
 		case EAppleAnchorType::PlaneAnchor:
 		{
 			NewAnchorDebugName = FString::Printf(TEXT("PLN-%02d"), LastTrackedGeometry_DebugId++);
 			UARPlaneGeometry* NewGeo = NewObject<UARPlaneGeometry>();
-			NewGeo->UpdateTrackedGeometry(ARComponent.ToSharedRef(), GameThreadFrameNumber, UpdateTimestamp, AnchorData->Transform, GetARCompositionComponent()->GetAlignmentTransform(), AnchorData->Center, AnchorData->Extent);
+			NewGeo->UpdateTrackedGeometry(ARComponent.ToSharedRef(), AnchorData->FrameNumber, AnchorData->Timestamp, AnchorData->Transform, GetARCompositionComponent()->GetAlignmentTransform(), AnchorData->Center, AnchorData->Extent);
 			NewGeometry = NewGeo;
 			break;
 		}
@@ -1578,7 +1582,7 @@ void FAppleARKitSystem::SessionDidAddAnchors_Internal( TSharedRef<FAppleARKitAnc
 		{
 			NewAnchorDebugName = FString::Printf(TEXT("FACE-%02d"), LastTrackedGeometry_DebugId++);
 			UARFaceGeometry* NewGeo = NewObject<UARFaceGeometry>();
-			NewGeo->UpdateFaceGeometry(ARComponent.ToSharedRef(), GameThreadFrameNumber, UpdateTimestamp, AnchorData->Transform, GetARCompositionComponent()->GetAlignmentTransform(), AnchorData->BlendShapes, AnchorData->FaceVerts, AnchorData->FaceIndices, AnchorData->LeftEyeTransform, AnchorData->RightEyeTransform, AnchorData->LookAtTarget);
+			NewGeo->UpdateFaceGeometry(ARComponent.ToSharedRef(), AnchorData->FrameNumber, AnchorData->Timestamp, AnchorData->Transform, GetARCompositionComponent()->GetAlignmentTransform(), AnchorData->BlendShapes, AnchorData->FaceVerts, AnchorData->FaceIndices, AnchorData->LeftEyeTransform, AnchorData->RightEyeTransform, AnchorData->LookAtTarget);
 			NewGeo->SetTrackingState(EARTrackingState::Tracking);
 			// @todo JoeG -- remove in 4.22
 PRAGMA_DISABLE_DEPRECATION_WARNINGS
@@ -1593,7 +1597,7 @@ PRAGMA_ENABLE_DEPRECATION_WARNINGS
 			UARTrackedImage* NewImage = NewObject<UARTrackedImage>();
 			UARCandidateImage** CandidateImage = CandidateImages.Find(AnchorData->DetectedAnchorName);
 			ensure(CandidateImage != nullptr);
-			NewImage->UpdateTrackedGeometry(ARComponent.ToSharedRef(), GameThreadFrameNumber, UpdateTimestamp, AnchorData->Transform, GetARCompositionComponent()->GetAlignmentTransform(), *CandidateImage);
+			NewImage->UpdateTrackedGeometry(ARComponent.ToSharedRef(), AnchorData->FrameNumber, AnchorData->Timestamp, AnchorData->Transform, GetARCompositionComponent()->GetAlignmentTransform(), *CandidateImage);
 			NewGeometry = NewImage;
 			break;
 		}
@@ -1601,7 +1605,7 @@ PRAGMA_ENABLE_DEPRECATION_WARNINGS
 		{
 			NewAnchorDebugName = FString::Printf(TEXT("ENV-%02d"), LastTrackedGeometry_DebugId++);
 			UAppleARKitEnvironmentCaptureProbe* NewProbe = NewObject<UAppleARKitEnvironmentCaptureProbe>();
-			NewProbe->UpdateEnvironmentCapture(ARComponent.ToSharedRef(), GameThreadFrameNumber, UpdateTimestamp, AnchorData->Transform, GetARCompositionComponent()->GetAlignmentTransform(), AnchorData->Extent, AnchorData->ProbeTexture);
+			NewProbe->UpdateEnvironmentCapture(ARComponent.ToSharedRef(), AnchorData->FrameNumber, AnchorData->Timestamp, AnchorData->Transform, GetARCompositionComponent()->GetAlignmentTransform(), AnchorData->Extent, AnchorData->ProbeTexture);
 			NewGeometry = NewProbe;
 			break;
 		}
@@ -1611,7 +1615,7 @@ PRAGMA_ENABLE_DEPRECATION_WARNINGS
 			UARTrackedObject* NewTrackedObject = NewObject<UARTrackedObject>();
 			UARCandidateObject** CandidateObject = CandidateObjects.Find(AnchorData->DetectedAnchorName);
 			ensure(CandidateObject != nullptr);
-			NewTrackedObject->UpdateTrackedGeometry(ARComponent.ToSharedRef(), GameThreadFrameNumber, UpdateTimestamp, AnchorData->Transform, GetARCompositionComponent()->GetAlignmentTransform(), *CandidateObject);
+			NewTrackedObject->UpdateTrackedGeometry(ARComponent.ToSharedRef(), AnchorData->FrameNumber, AnchorData->Timestamp, AnchorData->Transform, GetARCompositionComponent()->GetAlignmentTransform(), *CandidateObject);
 			NewGeometry = NewTrackedObject;
 			break;
 		}
@@ -1638,7 +1642,7 @@ void FAppleARKitSystem::SessionDidUpdateAnchors_Internal( TSharedRef<FAppleARKit
 	// If this object is valid, we are running a face session and we need to publish LiveLink data on the game thread
 	if (FaceARSupport != nullptr && AnchorData->AnchorType == EAppleAnchorType::FaceAnchor)
 	{
-		FaceARSupport->PublishLiveLinkData(AnchorData, UpdateTimestamp, GameThreadFrameNumber);
+		FaceARSupport->PublishLiveLinkData(AnchorData);
 	}
 
 	UARTrackedGeometry** GeometrySearchResult = TrackedGeometries.Find(AnchorData->AnchorGUID);
@@ -1659,7 +1663,7 @@ void FAppleARKitSystem::SessionDidUpdateAnchors_Internal( TSharedRef<FAppleARKit
 		{
 			case EAppleAnchorType::Anchor:
 			{
-				FoundGeometry->UpdateTrackedGeometry(ARComponent.ToSharedRef(), GameThreadFrameNumber, UpdateTimestamp, AnchorData->Transform, GetARCompositionComponent()->GetAlignmentTransform());
+				FoundGeometry->UpdateTrackedGeometry(ARComponent.ToSharedRef(), AnchorData->FrameNumber, AnchorData->Timestamp, AnchorData->Transform, GetARCompositionComponent()->GetAlignmentTransform());
 				for (UARPin* Pin : PinsToUpdate)
 				{
 					const FTransform Pin_LocalToTrackingTransform_PostUpdate = Pin->GetLocalToTrackingTransform_NoAlignment() * AnchorDeltaTransform;
@@ -1672,7 +1676,7 @@ void FAppleARKitSystem::SessionDidUpdateAnchors_Internal( TSharedRef<FAppleARKit
 			{
 				if (UARPlaneGeometry* PlaneGeo = Cast<UARPlaneGeometry>(FoundGeometry))
 				{
-					PlaneGeo->UpdateTrackedGeometry(ARComponent.ToSharedRef(), GameThreadFrameNumber, UpdateTimestamp, AnchorData->Transform, GetARCompositionComponent()->GetAlignmentTransform(), AnchorData->Center, AnchorData->Extent, AnchorData->BoundaryVerts, nullptr);
+					PlaneGeo->UpdateTrackedGeometry(ARComponent.ToSharedRef(), AnchorData->FrameNumber, AnchorData->Timestamp, AnchorData->Transform, GetARCompositionComponent()->GetAlignmentTransform(), AnchorData->Center, AnchorData->Extent, AnchorData->BoundaryVerts, nullptr);
 					for (UARPin* Pin : PinsToUpdate)
 					{
 						const FTransform Pin_LocalToTrackingTransform_PostUpdate = Pin->GetLocalToTrackingTransform_NoAlignment() * AnchorDeltaTransform;
@@ -1685,7 +1689,7 @@ void FAppleARKitSystem::SessionDidUpdateAnchors_Internal( TSharedRef<FAppleARKit
 			{
 				if (UARFaceGeometry* FaceGeo = Cast<UARFaceGeometry>(FoundGeometry))
 				{
-					FaceGeo->UpdateFaceGeometry(ARComponent.ToSharedRef(), GameThreadFrameNumber, UpdateTimestamp, AnchorData->Transform, GetARCompositionComponent()->GetAlignmentTransform(), AnchorData->BlendShapes, AnchorData->FaceVerts, AnchorData->FaceIndices, AnchorData->LeftEyeTransform, AnchorData->RightEyeTransform, AnchorData->LookAtTarget);
+					FaceGeo->UpdateFaceGeometry(ARComponent.ToSharedRef(), AnchorData->FrameNumber, AnchorData->Timestamp, AnchorData->Transform, GetARCompositionComponent()->GetAlignmentTransform(), AnchorData->BlendShapes, AnchorData->FaceVerts, AnchorData->FaceIndices, AnchorData->LeftEyeTransform, AnchorData->RightEyeTransform, AnchorData->LookAtTarget);
 					FaceGeo->SetTrackingState(AnchorData->bIsTracked ? EARTrackingState::Tracking : EARTrackingState::NotTracking);
 					// @todo JoeG -- remove this in 4.22
 PRAGMA_DISABLE_DEPRECATION_WARNINGS
@@ -1706,7 +1710,7 @@ PRAGMA_ENABLE_DEPRECATION_WARNINGS
 					UARCandidateImage** CandidateImage = CandidateImages.Find(AnchorData->DetectedAnchorName);
 					ensure(CandidateImage != nullptr);
 
-					ImageAnchor->UpdateTrackedGeometry(ARComponent.ToSharedRef(), GameThreadFrameNumber, UpdateTimestamp, AnchorData->Transform, GetARCompositionComponent()->GetAlignmentTransform(), *CandidateImage);
+					ImageAnchor->UpdateTrackedGeometry(ARComponent.ToSharedRef(), AnchorData->FrameNumber, AnchorData->Timestamp, AnchorData->Transform, GetARCompositionComponent()->GetAlignmentTransform(), *CandidateImage);
 					ImageAnchor->SetTrackingState(AnchorData->bIsTracked ? EARTrackingState::Tracking : EARTrackingState::NotTracking);
 					// @todo JoeG -- remove this in 4.22
 PRAGMA_DISABLE_DEPRECATION_WARNINGS
@@ -1725,7 +1729,7 @@ PRAGMA_ENABLE_DEPRECATION_WARNINGS
 				if (UAppleARKitEnvironmentCaptureProbe* ProbeAnchor = Cast<UAppleARKitEnvironmentCaptureProbe>(FoundGeometry))
 				{
 					// NOTE: The metal texture will be a different texture every time the cubemap is updated which requires a render resource flush
-					ProbeAnchor->UpdateEnvironmentCapture(ARComponent.ToSharedRef(), GameThreadFrameNumber, UpdateTimestamp, AnchorData->Transform, GetARCompositionComponent()->GetAlignmentTransform(), AnchorData->Extent, AnchorData->ProbeTexture);
+					ProbeAnchor->UpdateEnvironmentCapture(ARComponent.ToSharedRef(), AnchorData->FrameNumber, AnchorData->Timestamp, AnchorData->Transform, GetARCompositionComponent()->GetAlignmentTransform(), AnchorData->Extent, AnchorData->ProbeTexture);
 					for (UARPin* Pin : PinsToUpdate)
 					{
 						const FTransform Pin_LocalToTrackingTransform_PostUpdate = Pin->GetLocalToTrackingTransform_NoAlignment() * AnchorDeltaTransform;
diff --git a/Engine/Plugins/Runtime/AR/Apple/AppleARKit/Source/AppleARKit/Private/AppleARKitSystem.h b/Engine/Plugins/Runtime/AR/Apple/AppleARKit/Source/AppleARKit/Private/AppleARKitSystem.h
--- a/Engine/Plugins/Runtime/AR/Apple/AppleARKit/Source/AppleARKit/Private/AppleARKitSystem.h
+++ b/Engine/Plugins/Runtime/AR/Apple/AppleARKit/Source/AppleARKit/Private/AppleARKitSystem.h
@@ -175,11 +175,6 @@ private:
 	// PROPERTIES REPORTED TO FGCObject
 	//
 	
-	/** The ar frame number when LastReceivedFrame was last updated */
-	uint32 GameThreadFrameNumber;
-	/** The ar timestamp of when the LastReceivedFrame was last updated */
-	double GameThreadTimestamp;
-
 	// An int counter that provides a human-readable debug number for Tracked Geometries.
 	uint32 LastTrackedGeometry_DebugId;
 
diff --git a/Engine/Plugins/Runtime/AR/Apple/AppleARKit/Source/AppleARKit/Public/AppleARKitConversion.h b/Engine/Plugins/Runtime/AR/Apple/AppleARKit/Source/AppleARKit/Public/AppleARKitConversion.h
--- a/Engine/Plugins/Runtime/AR/Apple/AppleARKit/Source/AppleARKit/Public/AppleARKitConversion.h
+++ b/Engine/Plugins/Runtime/AR/Apple/AppleARKit/Source/AppleARKit/Public/AppleARKitConversion.h
@@ -16,6 +16,45 @@
 	#import <ARKit/ARKit.h>
 #endif
 
+/**
+ * This class handles incrementing the frame counter. This is needed in case there
+ * are both new anchors added and existing ones updated in the same "frame". When
+ * this happens, the frame number will remain the same due to the time between
+ * adds and updates being shorter than the arkit update rate.
+ */
+class FARKitFrameCounter
+{
+public:
+	static FARKitFrameCounter& Get()
+	{
+		static FARKitFrameCounter Singleton;
+		return Singleton;
+	}
+
+	uint32 GetFrameNumber()
+	{
+		uint32 FrameNum = FrameNumber;
+		double TimeSecs = FPlatformTime::Seconds();
+		// ARKit runs at 60hz
+		if (TimeSecs - LastFrameSecs > 0.016)
+		{
+			FrameNumber++;
+			LastFrameSecs = TimeSecs;
+		}
+		return FrameNum;
+	}
+
+private:
+	FARKitFrameCounter()
+		: FrameNumber(0)
+		, LastFrameSecs(0.0)
+	{
+	}
+
+	uint32 FrameNumber;
+	double LastFrameSecs;
+};
+
 
 #define AR_SAVE_WORLD_KEY 0x505A474A
 #define AR_SAVE_WORLD_VER 1
@@ -234,7 +273,7 @@ struct FAppleARKitAnchorData
 	{
 	}
 
-	FAppleARKitAnchorData(FGuid InAnchorGuid, FTransform InTransform, FARBlendShapeMap InBlendShapes, TArray<FVector> InFaceVerts, FTransform InLeftEyeTransform, FTransform InRightEyeTransform, FVector InLookAtTarget)
+	FAppleARKitAnchorData(FGuid InAnchorGuid, FTransform InTransform, FARBlendShapeMap InBlendShapes, TArray<FVector> InFaceVerts, FTransform InLeftEyeTransform, FTransform InRightEyeTransform, FVector InLookAtTarget, double InTimestamp, uint32 InFrameNumber)
 		: Transform( InTransform )
 		, AnchorType( EAppleAnchorType::FaceAnchor )
 		, AnchorGUID( InAnchorGuid )
@@ -243,6 +282,8 @@ struct FAppleARKitAnchorData
 		, LeftEyeTransform( InLeftEyeTransform )
 		, RightEyeTransform( InRightEyeTransform )
 		, LookAtTarget( InLookAtTarget )
+		, Timestamp( InTimestamp )
+		, FrameNumber( InFrameNumber )
 	{
 	}
 
@@ -282,6 +323,8 @@ struct FAppleARKitAnchorData
 	FTransform LeftEyeTransform;
 	FTransform RightEyeTransform;
 	FVector LookAtTarget;
+	double Timestamp;
+	uint32 FrameNumber;
 
 	/** Only valid for tracked real world objects (face, images) */
 	bool bIsTracked;
diff --git a/Engine/Plugins/Runtime/AR/Apple/AppleARKit/Source/AppleARKit/Public/AppleARKitFaceSupport.h b/Engine/Plugins/Runtime/AR/Apple/AppleARKit/Source/AppleARKit/Public/AppleARKitFaceSupport.h
--- a/Engine/Plugins/Runtime/AR/Apple/AppleARKit/Source/AppleARKit/Public/AppleARKitFaceSupport.h
+++ b/Engine/Plugins/Runtime/AR/Apple/AppleARKit/Source/AppleARKit/Public/AppleARKitFaceSupport.h
@@ -25,26 +25,22 @@ public:
 	 * Converts a set of generic ARAnchors into their face anchor equivalents without exposing the main code to the face APIs
 	 *
 	 * @param NewAnchors the list of anchors to convert to our intermediate format
-	 * @param Timestamp the timestamp of this update
-	 * @param FrameNumber the frame number for this update
 	 * @param AdjustBy the additional rotation to apply to put the rotation in the proper space (camera alignment only)
 	 * @param UpdateSetting whether to just update curves or geo too
 	 *
 	 * @return the set of face anchors to dispatch
 	 */
-	virtual TArray<TSharedPtr<FAppleARKitAnchorData>> MakeAnchorData(NSArray<ARAnchor*>* NewAnchors, double Timestamp, uint32 FrameNumber, const FRotator& AdjustBy, EARFaceTrackingUpdate UpdateSetting) { return TArray<TSharedPtr<FAppleARKitAnchorData>>(); }
+	virtual TArray<TSharedPtr<FAppleARKitAnchorData>> MakeAnchorData(NSArray<ARAnchor*>* NewAnchors, const FRotator& AdjustBy, EARFaceTrackingUpdate UpdateSetting) { return TArray<TSharedPtr<FAppleARKitAnchorData>>(); }
 
 	/**
 	 * Publishes any face AR data that needs to be sent to LiveLink. Done as a separate step because MakeAnchorData is called
 	 * on an arbitrary thread and we can't access UObjects there safely
 	 *
 	 * @param AnchorList the list of anchors to publish to LiveLink
-	 * @param Timestamp the timestamp of this update
-	 * @param FrameNumber the frame number for this update
 	 *
 	 * @return the set of face anchors to dispatch
 	 */
-	virtual void PublishLiveLinkData(TSharedPtr<FAppleARKitAnchorData> Anchor, double Timestamp, uint32 FrameNumber) { }
+	virtual void PublishLiveLinkData(TSharedPtr<FAppleARKitAnchorData> Anchor) { }
 
 	/**
 	 * Creates a face ar specific configuration object if that is requested without exposing the main code to the face APIs
diff --git a/Engine/Plugins/Runtime/AR/Apple/AppleARKitFaceSupport/Source/AppleARKitFaceSupport/Private/AppleARKitFaceSupportImpl.cpp b/Engine/Plugins/Runtime/AR/Apple/AppleARKitFaceSupport/Source/AppleARKitFaceSupport/Private/AppleARKitFaceSupportImpl.cpp
--- a/Engine/Plugins/Runtime/AR/Apple/AppleARKitFaceSupport/Source/AppleARKitFaceSupport/Private/AppleARKitFaceSupportImpl.cpp
+++ b/Engine/Plugins/Runtime/AR/Apple/AppleARKitFaceSupport/Source/AppleARKitFaceSupport/Private/AppleARKitFaceSupportImpl.cpp
@@ -32,6 +32,8 @@ static TSharedPtr<FAppleARKitAnchorData> MakeAnchorData(ARAnchor* Anchor, const
 			LookAtTarget = FAppleARKitConversion::ToFVector(FaceAnchor.lookAtPoint);
 		}
 #endif
+		uint32 FrameNumber = FARKitFrameCounter::Get().GetFrameNumber();
+		double Timestamp = FPlatformTime::Seconds();
 		NewAnchor = MakeShared<FAppleARKitAnchorData>(
 			FAppleARKitConversion::ToFGuid(FaceAnchor.identifier),
 			FAppleARKitConversion::ToFTransform(FaceAnchor.transform, AdjustBy),
@@ -39,7 +41,9 @@ static TSharedPtr<FAppleARKitAnchorData> MakeAnchorData(ARAnchor* Anchor, const
 			UpdateSetting == EARFaceTrackingUpdate::CurvesAndGeo ? ToVertexBuffer(FaceAnchor.geometry.vertices, FaceAnchor.geometry.vertexCount) : TArray<FVector>(),
 			LeftEyeTransform,
 			RightEyeTransform,
-			LookAtTarget
+			LookAtTarget,
+			Timestamp,
+			FrameNumber++
 		);
         // Only convert from 16bit to 32bit once
         if (UpdateSetting == EARFaceTrackingUpdate::CurvesAndGeo && FAppleARKitAnchorData::FaceIndices.Num() == 0)
@@ -56,6 +60,8 @@ static TSharedPtr<FAppleARKitAnchorData> MakeAnchorData(ARAnchor* Anchor, const
 
 FAppleARKitFaceSupport::FAppleARKitFaceSupport()
 {
+	// Generate our device id
+	LocalDeviceId = FName(*FPlatformMisc::GetDeviceId());
 }
 
 FAppleARKitFaceSupport::~FAppleARKitFaceSupport()
@@ -106,7 +112,7 @@ ARConfiguration* FAppleARKitFaceSupport::ToARConfiguration(UARSessionConfig* Ses
 	return SessionConfiguration;
 }
 
-TArray<TSharedPtr<FAppleARKitAnchorData>> FAppleARKitFaceSupport::MakeAnchorData(NSArray<ARAnchor*>* Anchors, double Timestamp, uint32 FrameNumber, const FRotator& AdjustBy, EARFaceTrackingUpdate UpdateSetting)
+TArray<TSharedPtr<FAppleARKitAnchorData>> FAppleARKitFaceSupport::MakeAnchorData(NSArray<ARAnchor*>* Anchors, const FRotator& AdjustBy, EARFaceTrackingUpdate UpdateSetting)
 {
 	TArray<TSharedPtr<FAppleARKitAnchorData>> AnchorList;
 
@@ -122,7 +128,7 @@ TArray<TSharedPtr<FAppleARKitAnchorData>> FAppleARKitFaceSupport::MakeAnchorData
 	return AnchorList;
 }
 
-void FAppleARKitFaceSupport::PublishLiveLinkData(TSharedPtr<FAppleARKitAnchorData> Anchor, double Timestamp, uint32 FrameNumber)
+void FAppleARKitFaceSupport::PublishLiveLinkData(TSharedPtr<FAppleARKitAnchorData> Anchor)
 {
 	static bool bNeedsInit = true;
 	if (bNeedsInit)
@@ -143,7 +149,7 @@ void FAppleARKitFaceSupport::PublishLiveLinkData(TSharedPtr<FAppleARKitAnchorDat
 
 	if (LiveLinkSource.IsValid())
 	{
-        LiveLinkSource->PublishBlendShapes(FaceTrackingLiveLinkSubjectName, Timestamp, FrameNumber, Anchor->BlendShapes);
+        LiveLinkSource->PublishBlendShapes(FaceTrackingLiveLinkSubjectName, Anchor->Timestamp, Anchor->FrameNumber, Anchor->BlendShapes, LocalDeviceId);
 	}
 }
 
diff --git a/Engine/Plugins/Runtime/AR/Apple/AppleARKitFaceSupport/Source/AppleARKitFaceSupport/Private/AppleARKitFaceSupportImpl.h b/Engine/Plugins/Runtime/AR/Apple/AppleARKitFaceSupport/Source/AppleARKitFaceSupport/Private/AppleARKitFaceSupportImpl.h
--- a/Engine/Plugins/Runtime/AR/Apple/AppleARKitFaceSupport/Source/AppleARKitFaceSupport/Private/AppleARKitFaceSupportImpl.h
+++ b/Engine/Plugins/Runtime/AR/Apple/AppleARKitFaceSupport/Source/AppleARKitFaceSupport/Private/AppleARKitFaceSupportImpl.h
@@ -22,9 +22,9 @@ public:
 private:
 #if SUPPORTS_ARKIT_1_0
 	// ~IAppleARKitFaceSupport
-	virtual TArray<TSharedPtr<FAppleARKitAnchorData>> MakeAnchorData(NSArray<ARAnchor*>* NewAnchors, double Timestamp, uint32 FrameNumber, const FRotator& AdjustBy, EARFaceTrackingUpdate UpdateSetting) override;
+	virtual TArray<TSharedPtr<FAppleARKitAnchorData>> MakeAnchorData(NSArray<ARAnchor*>* NewAnchors, const FRotator& AdjustBy, EARFaceTrackingUpdate UpdateSetting) override;
 	virtual ARConfiguration* ToARConfiguration(UARSessionConfig* InSessionConfig) override;
-	virtual void PublishLiveLinkData(TSharedPtr<FAppleARKitAnchorData> Anchor, double Timestamp, uint32 FrameNumber) override;
+	virtual void PublishLiveLinkData(TSharedPtr<FAppleARKitAnchorData> Anchor) override;
 	virtual bool DoesSupportFaceAR() override;
 #if SUPPORTS_ARKIT_1_5
 	virtual TArray<FARVideoFormat> ToARConfiguration() override;
@@ -36,5 +36,7 @@ private:
 	TSharedPtr<ILiveLinkSourceARKit> LiveLinkSource;
 	/** Copied from the UARSessionConfig project settings object */
 	FName FaceTrackingLiveLinkSubjectName;
+	/** The id of this device */
+	FName LocalDeviceId;
 };
 
diff --git a/Engine/Plugins/Runtime/AR/Apple/AppleARKitFaceSupport/Source/AppleARKitFaceSupport/Private/AppleARKitLiveLinkSource.cpp b/Engine/Plugins/Runtime/AR/Apple/AppleARKitFaceSupport/Source/AppleARKitFaceSupport/Private/AppleARKitLiveLinkSource.cpp
--- a/Engine/Plugins/Runtime/AR/Apple/AppleARKitFaceSupport/Source/AppleARKitFaceSupport/Private/AppleARKitLiveLinkSource.cpp
+++ b/Engine/Plugins/Runtime/AR/Apple/AppleARKitFaceSupport/Source/AppleARKitFaceSupport/Private/AppleARKitLiveLinkSource.cpp
@@ -134,12 +134,6 @@ void FAppleARKitLiveLinkSource::PublishBlendShapes(FName SubjectName, double Tim
 	// This code touches UObjects so needs to be run only in the game thread
 	check(IsInGameThread());
 
-	// Generate our device id if it wasn't passed in (local only case)
-	if (DeviceId == NAME_None)
-	{
-		DeviceId = FName(*FPlatformMisc::GetDeviceId());
-	}
-
 	FName* LastSubjectNameForDeviceId = DeviceToLastSubjectNameMap.Find(DeviceId);
 	// Is this a new device and subject pair?
 	if (LastSubjectNameForDeviceId == nullptr)
