Added code to expose the Apple ARKit settings object to the project settings
Added new options for the file writer support for logging face curve data as it is tracked

diff --git a/Engine/Plugins/Runtime/AR/Apple/AppleARKit/Source/AppleARKit/Private/AppleARKitModule.cpp b/Engine/Plugins/Runtime/AR/Apple/AppleARKit/Source/AppleARKit/Private/AppleARKitModule.cpp
--- a/Engine/Plugins/Runtime/AR/Apple/AppleARKit/Source/AppleARKit/Private/AppleARKitModule.cpp
+++ b/Engine/Plugins/Runtime/AR/Apple/AppleARKit/Source/AppleARKit/Private/AppleARKitModule.cpp
@@ -8,6 +8,13 @@
 #include "CoreGlobals.h"
 #include "Misc/CoreDelegates.h"
 
+#if WITH_EDITOR
+	#include "ISettingsModule.h"
+	#include "ISettingsSection.h"
+	#include "AppleARKitSettings.h"
+#endif
+
+#define LOCTEXT_NAMESPACE "ARKit"
 
 TWeakPtr<class FAppleARKitSystem, ESPMode::ThreadSafe> FAppleARKitARKitSystemPtr;
 
@@ -40,10 +47,17 @@ void FAppleARKitModule::StartupModule()
 	IHeadMountedDisplayModule::StartupModule();
 
 	FCoreDelegates::OnPreExit.AddRaw(this, &FAppleARKitModule::PreExit);
+
+#if WITH_EDITOR
+	FCoreDelegates::OnPostEngineInit.AddRaw(this, &FAppleARKitModule::PostEngineInit);
+#endif
 }
 
 void FAppleARKitModule::PreExit()
 {
+#if WITH_EDITOR
+	UnregisterSettings();
+#endif
 	if (FAppleARKitARKitSystemPtr.IsValid())
 	{
 		FAppleARKitARKitSystemPtr.Pin()->Shutdown();
@@ -56,6 +70,38 @@ void FAppleARKitModule::ShutdownModule()
 	IHeadMountedDisplayModule::ShutdownModule();
 }
 
+#if WITH_EDITOR
+void FAppleARKitModule::PostEngineInit()
+{
+	RegisterSettings();
+}
+
+void FAppleARKitModule::RegisterSettings()
+{
+	{
+		ISettingsModule* SettingsModule = FModuleManager::GetModulePtr<ISettingsModule>("Settings");
+
+		if (SettingsModule != nullptr)
+		{
+			SettingsModule->RegisterSettings("Project", "Plugins", "Apple ARKit",
+				LOCTEXT("ARKitSettingsName", "Apple ARKit"),
+				LOCTEXT("ARKitSettingsDescription", "Configure the Apple ARKit plug-in."),
+				GetMutableDefault<UAppleARKitSettings>()
+			);
+		}
+	}
+}
+
+void FAppleARKitModule::UnregisterSettings()
+{
+	ISettingsModule* SettingsModule = FModuleManager::GetModulePtr<ISettingsModule>("Settings");
+	if (SettingsModule)
+	{
+		SettingsModule->UnregisterSettings("Project", "Plugins", "Apple ARKit");
+	}
+}
+#endif
+
 
 IMPLEMENT_MODULE(FAppleARKitModule, AppleARKit);
 
diff --git a/Engine/Plugins/Runtime/AR/Apple/AppleARKit/Source/AppleARKit/Private/AppleARKitModule.h b/Engine/Plugins/Runtime/AR/Apple/AppleARKit/Source/AppleARKit/Private/AppleARKitModule.h
--- a/Engine/Plugins/Runtime/AR/Apple/AppleARKit/Source/AppleARKit/Private/AppleARKitModule.h
+++ b/Engine/Plugins/Runtime/AR/Apple/AppleARKit/Source/AppleARKit/Private/AppleARKitModule.h
@@ -22,6 +22,12 @@ public:
 	virtual void ShutdownModule() override;
 
     void PreExit();
+
+#if WITH_EDITOR
+	void PostEngineInit();
+    void RegisterSettings();
+    void UnregisterSettings();
+#endif
 };
 
 DECLARE_LOG_CATEGORY_EXTERN(LogAppleARKit, Log, All);
diff --git a/Engine/Plugins/Runtime/AR/Apple/AppleARKit/Source/AppleARKit/Public/AppleARKitSettings.h b/Engine/Plugins/Runtime/AR/Apple/AppleARKit/Source/AppleARKit/Public/AppleARKitSettings.h
--- a/Engine/Plugins/Runtime/AR/Apple/AppleARKit/Source/AppleARKit/Public/AppleARKitSettings.h
+++ b/Engine/Plugins/Runtime/AR/Apple/AppleARKit/Source/AppleARKit/Public/AppleARKitSettings.h
@@ -8,14 +8,29 @@
 
 #include "AppleARKitSettings.generated.h"
 
-UCLASS(Config=Engine)
-class APPLEARKIT_API UAppleARKitSettings : public UObject
+UENUM(BlueprintType, Category="AR AugmentedReality", meta=(Experimental))
+enum class EARFaceTrackingFileWriterType : uint8
+{
+	/** Disables creation of a file writer */
+	None,
+	/** Comma delimited file, one row per captured frame */
+	CSV,
+	/** JSON object array, one frame object per captured frame */
+	JSON
+};
+
+
+UCLASS(Config=Engine, defaultconfig)
+class APPLEARKIT_API UAppleARKitSettings :
+	public UObject
 {
 	GENERATED_BODY()
 
 public:
 	UAppleARKitSettings()
 		: bEnableLiveLinkForFaceTracking(false)
+		, bFaceTrackingWriteEachFrame(false)
+		, FaceTrackingFileWriterType(EARFaceTrackingFileWriterType::None)
 		, LiveLinkPublishingPort(11111)
 		, DefaultFaceTrackingLiveLinkSubjectName(FName("iPhoneXFaceAR"))
 		, DefaultFaceTrackingDirection(EARFaceTrackingDirection::FaceRelative)
@@ -29,6 +44,14 @@ public:
 	UPROPERTY(Config, EditAnywhere, Category="AR Settings")
 	bool bEnableLiveLinkForFaceTracking;
 
+	/** Whether to publish each frame or when the "FaceAR WriteCurveFile */
+	UPROPERTY(Config, EditAnywhere, Category="AR Settings")
+	bool bFaceTrackingWriteEachFrame;
+
+	/** The type of face AR publisher that writes to disk to create */
+	UPROPERTY(Config, EditAnywhere, Category="AR Settings")
+	EARFaceTrackingFileWriterType FaceTrackingFileWriterType;
+
 	/** The port to use when listening/sending LiveLink face blend shapes via the network */
 	UPROPERTY(Config, EditAnywhere, Category="AR Settings")
 	int32 LiveLinkPublishingPort;
@@ -44,7 +67,7 @@ public:
 	/** Whether to adjust thread priorities during an AR session or not */
 	UPROPERTY(Config, EditAnywhere, Category="AR Settings")
 	bool bAdjustThreadPrioritiesDuringARSession;
-	
+
 	/** The game thread priority to change to when an AR session is running, default is 47 */
 	UPROPERTY(Config, EditAnywhere, Category="AR Settings")
 	int32 GameThreadPriorityOverride;
diff --git a/Engine/Plugins/Runtime/AR/Apple/AppleARKitFaceSupport/Source/AppleARKitFaceSupport/Private/AppleARKitLiveLinkSource.cpp b/Engine/Plugins/Runtime/AR/Apple/AppleARKitFaceSupport/Source/AppleARKitFaceSupport/Private/AppleARKitLiveLinkSource.cpp
--- a/Engine/Plugins/Runtime/AR/Apple/AppleARKitFaceSupport/Source/AppleARKitFaceSupport/Private/AppleARKitLiveLinkSource.cpp
+++ b/Engine/Plugins/Runtime/AR/Apple/AppleARKitFaceSupport/Source/AppleARKitFaceSupport/Private/AppleARKitLiveLinkSource.cpp
@@ -85,14 +85,12 @@ TSharedPtr<IARKitBlendShapePublisher, ESPMode::ThreadSafe> FAppleARKitLiveLinkSo
 TSharedPtr<IARKitBlendShapePublisher, ESPMode::ThreadSafe> FAppleARKitLiveLinkSourceFactory::CreateLiveLinkLocalFileWriter()
 {
 	TSharedPtr<IARKitBlendShapePublisher, ESPMode::ThreadSafe> LocalFileWriter;
-	bool bWantsFaceTrackingFileWriter = false;
-	GConfig->GetBool(TEXT("/Script/AppleARKit.AppleARKitSettings"), TEXT("bWantsFaceTrackingFileWriter"), bWantsFaceTrackingFileWriter, GEngineIni);
-	if (bWantsFaceTrackingFileWriter)
+	FString FileWriterType(TEXT("None"));
+	GConfig->GetString(TEXT("/Script/AppleARKit.AppleARKitSettings"), TEXT("FaceTrackingFileWriterType"), FileWriterType, GEngineIni);
+	if (FileWriterType != TEXT("None"))
 	{
-		bool bCsvOrJson = true;
-		GConfig->GetBool(TEXT("/Script/AppleARKit.AppleARKitSettings"), TEXT("bWantsCsvOrJsonFileWriter"), bCsvOrJson, GEngineIni);
 		FAppleARKitLiveLinkFileWriter* FileWriter = nullptr;
-		if (bCsvOrJson)
+		if (FileWriterType == TEXT("CSV"))
 		{
 			FileWriter = new FAppleARKitLiveLinkFileWriterCsv();
 		}
