Added support to write out the ARKit camera image each frame
Added project settings to control the compression quality, image scale, and image rotation when writing out the camera images

diff --git a/Engine/Plugins/Runtime/AR/Apple/AppleARKit/Source/AppleARKit/Private/AppleARKitSystem.cpp b/Engine/Plugins/Runtime/AR/Apple/AppleARKit/Source/AppleARKit/Private/AppleARKitSystem.cpp
--- a/Engine/Plugins/Runtime/AR/Apple/AppleARKit/Source/AppleARKit/Private/AppleARKitSystem.cpp
+++ b/Engine/Plugins/Runtime/AR/Apple/AppleARKit/Source/AppleARKit/Private/AppleARKitSystem.cpp
@@ -18,7 +18,9 @@
 #include "ARLightEstimate.h"
 #include "ARTraceResult.h"
 #include "ARPin.h"
+#include "Async/Async.h"
 #include "HAL/ThreadSafeCounter.h"
+#include "Misc/FileHelper.h"
 
 // To separate out the face ar library linkage from standard ar apps
 #include "AppleARKitFaceSupport.h"
@@ -1196,7 +1198,7 @@ bool FAppleARKitSystem::Run(UARSessionConfig* SessionConfig)
 		GameThreadFrame = TSharedPtr<FAppleARKitFrame, ESPMode::ThreadSafe>();
 		LastReceivedFrame = TSharedPtr<FAppleARKitFrame, ESPMode::ThreadSafe>();
 	}
-	
+
 	// Make sure this is set at session start, because there are timing issues with using only the delegate approach
 	if (DeviceOrientation == EScreenOrientation::Unknown)
 	{
@@ -1204,6 +1206,12 @@ bool FAppleARKitSystem::Run(UARSessionConfig* SessionConfig)
 	}
 
 #if SUPPORTS_ARKIT_1_0
+	// Set this based upon the project settings
+	bShouldWriteCameraImagePerFrame = GetDefault<UAppleARKitSettings>()->bShouldWriteCameraImagePerFrame;
+	WrittenCameraImageScale = GetDefault<UAppleARKitSettings>()->WrittenCameraImageScale;
+	WrittenCameraImageRotation = GetDefault<UAppleARKitSettings>()->WrittenCameraImageRotation;
+	WrittenCameraImageQuality = GetDefault<UAppleARKitSettings>()->WrittenCameraImageQuality;
+
 	if (FAppleARKitAvailability::SupportsARKit10())
 	{
 		ARSessionRunOptions options = 0;
@@ -1372,6 +1380,12 @@ void FAppleARKitSystem::SessionDidUpdateFrame_DelegateThread(TSharedPtr< FAppleA
 	}
 	{
 		UpdateARKitPerfStats();
+#if SUPPORTS_ARKIT_1_0
+		if (bShouldWriteCameraImagePerFrame)
+		{
+			WriteCameraImageToDisk(Frame->CameraImage);
+		}
+#endif
 	}
 }
 			
@@ -1866,6 +1880,29 @@ void FAppleARKitSystem::UpdateARKitPerfStats()
 #endif
 }
 
+#if SUPPORTS_ARKIT_1_0
+void FAppleARKitSystem::WriteCameraImageToDisk(CVPixelBufferRef PixelBuffer)
+{
+	int32 ImageQuality = WrittenCameraImageQuality;
+	float ImageScale = WrittenCameraImageScale;
+	ETextureRotationDirection ImageRotation = WrittenCameraImageRotation;
+	CIImage* SourceImage = [[CIImage alloc] initWithCVPixelBuffer: PixelBuffer];
+	AsyncTask(ENamedThreads::AnyBackgroundThreadNormalTask, [SourceImage, ImageQuality, ImageScale, ImageRotation]()
+	{
+		TArray<uint8> JpegBytes;
+		IAppleImageUtilsPlugin::Get().ConvertToJPEG(SourceImage, JpegBytes, ImageQuality, true, true, ImageScale, ImageRotation);
+		[SourceImage release];
+		// Build a unique file name
+		FDateTime DateTime = FDateTime::UtcNow();
+		static FString UserDir = FPlatformProcess::UserDir();
+		FString FileName = FString::Printf(TEXT("%sCameraImages/Image_%d-%d-%d-%d-%d-%d-%d.jpeg"), *UserDir,
+			DateTime.GetYear(), DateTime.GetMonth(), DateTime.GetDay(), DateTime.GetHour(), DateTime.GetMinute(), DateTime.GetSecond(), DateTime.GetMillisecond());
+		// Write the jpeg to disk
+		FFileHelper::SaveArrayToFile(JpegBytes, *FileName);
+	});
+}
+#endif
+
 
 namespace AppleARKitSupport
 {
diff --git a/Engine/Plugins/Runtime/AR/Apple/AppleARKit/Source/AppleARKit/Private/AppleARKitSystem.h b/Engine/Plugins/Runtime/AR/Apple/AppleARKit/Source/AppleARKit/Private/AppleARKitSystem.h
--- a/Engine/Plugins/Runtime/AR/Apple/AppleARKit/Source/AppleARKit/Private/AppleARKitSystem.h
+++ b/Engine/Plugins/Runtime/AR/Apple/AppleARKit/Source/AppleARKit/Private/AppleARKitSystem.h
@@ -94,6 +94,14 @@ private:
 	void UpdatePoses();
 	void UpdateFrame();
 	void CalcTrackingToWorldRotation();
+#if SUPPORTS_ARKIT_1_0
+	bool bShouldWriteCameraImagePerFrame;
+	float WrittenCameraImageScale;
+	int32 WrittenCameraImageQuality;
+	ETextureRotationDirection WrittenCameraImageRotation;
+	/** Asynchronously writes a JPEG to disk */
+	void WriteCameraImageToDisk(CVPixelBufferRef PixelBuffer);
+#endif
 
 public:
 	// Session delegate callbacks
diff --git a/Engine/Plugins/Runtime/AR/Apple/AppleARKit/Source/AppleARKit/Public/AppleARKitSettings.h b/Engine/Plugins/Runtime/AR/Apple/AppleARKit/Source/AppleARKit/Public/AppleARKitSettings.h
--- a/Engine/Plugins/Runtime/AR/Apple/AppleARKit/Source/AppleARKit/Public/AppleARKitSettings.h
+++ b/Engine/Plugins/Runtime/AR/Apple/AppleARKit/Source/AppleARKit/Public/AppleARKitSettings.h
@@ -6,6 +6,8 @@
 #include "UObject/Object.h"
 #include "ARTrackable.h"
 
+#include "AppleImageUtilsTypes.h"
+
 #include "AppleARKitSettings.generated.h"
 
 UENUM(BlueprintType, Category="AR AugmentedReality", meta=(Experimental))
@@ -31,6 +33,9 @@ public:
 		: bEnableLiveLinkForFaceTracking(false)
 		, bFaceTrackingWriteEachFrame(false)
 		, FaceTrackingFileWriterType(EARFaceTrackingFileWriterType::None)
+		, bShouldWriteCameraImagePerFrame(false)
+		, WrittenCameraImageScale(1.f)
+		, WrittenCameraImageQuality(85)
 		, LiveLinkPublishingPort(11111)
 		, DefaultFaceTrackingLiveLinkSubjectName(FName("iPhoneXFaceAR"))
 		, DefaultFaceTrackingDirection(EARFaceTrackingDirection::FaceRelative)
@@ -52,6 +57,22 @@ public:
 	UPROPERTY(Config, EditAnywhere, Category="AR Settings")
 	EARFaceTrackingFileWriterType FaceTrackingFileWriterType;
 
+	/** Whether to publish the camera image each frame */
+	UPROPERTY(Config, EditAnywhere, Category="AR Settings")
+	bool bShouldWriteCameraImagePerFrame;
+
+	/** The scale to write the images at. Used to reduce data footprint */
+	UPROPERTY(Config, EditAnywhere, Category="AR Settings")
+	float WrittenCameraImageScale;
+
+	/** The quality setting to generate the jpeg images at. Defaults to 85, which is "high quality". Lower values reduce data footprint */
+	UPROPERTY(Config, EditAnywhere, Category="AR Settings")
+	int32 WrittenCameraImageQuality;
+
+	/** Defaults to none. Use Right when in portrait mode */
+	UPROPERTY(Config, EditAnywhere, Category="AR Settings")
+	ETextureRotationDirection WrittenCameraImageRotation;
+
 	/** The port to use when listening/sending LiveLink face blend shapes via the network */
 	UPROPERTY(Config, EditAnywhere, Category="AR Settings")
 	int32 LiveLinkPublishingPort;
