Fix for non-Mac compilation error

diff --git a/Engine/Plugins/Runtime/AR/Apple/AppleARKit/Source/AppleARKit/Private/AppleARKitSystem.cpp b/Engine/Plugins/Runtime/AR/Apple/AppleARKit/Source/AppleARKit/Private/AppleARKitSystem.cpp
--- a/Engine/Plugins/Runtime/AR/Apple/AppleARKit/Source/AppleARKit/Private/AppleARKitSystem.cpp
+++ b/Engine/Plugins/Runtime/AR/Apple/AppleARKit/Source/AppleARKit/Private/AppleARKitSystem.cpp
@@ -1924,6 +1924,27 @@ namespace AppleARKitSupport
 	}
 }
 
+UTimecodeProvider* UAppleARKitSettings::GetTimecodeProvider()
+{
+	const FString& ProviderName = GetDefault<UAppleARKitSettings>()->ARKitTimecodeProvider;
+	UTimecodeProvider* TimecodeProvider = FindObject<UTimecodeProvider>(GEngine, *ProviderName);
+	if (TimecodeProvider == nullptr)
+	{
+		// Try to load the class that was requested
+		UClass* Class = LoadClass<UTimecodeProvider>(nullptr, *ProviderName);
+		if (Class != nullptr)
+		{
+			TimecodeProvider = NewObject<UTimecodeProvider>(GEngine, Class);
+		}
+	}
+	// Create the default one if this failed for some reason
+	if (TimecodeProvider == nullptr)
+	{
+		TimecodeProvider = NewObject<UTimecodeProvider>(GEngine, UAppleARKitTimecodeProvider::StaticClass());
+	}
+	return TimecodeProvider;
+}
+
 #if PLATFORM_IOS
 	#pragma clang diagnostic pop
 #endif
diff --git a/Engine/Plugins/Runtime/AR/Apple/AppleARKit/Source/AppleARKit/Public/AppleARKitSettings.h b/Engine/Plugins/Runtime/AR/Apple/AppleARKit/Source/AppleARKit/Public/AppleARKitSettings.h
--- a/Engine/Plugins/Runtime/AR/Apple/AppleARKit/Source/AppleARKit/Public/AppleARKitSettings.h
+++ b/Engine/Plugins/Runtime/AR/Apple/AppleARKit/Source/AppleARKit/Public/AppleARKitSettings.h
@@ -107,24 +107,5 @@ public:
 	UPROPERTY(Config, EditAnywhere, Category="AR Settings")
 	FString ARKitTimecodeProvider;
 
-	static UTimecodeProvider* GetTimecodeProvider()
-	{
-		const FString& ProviderName = GetDefault<UAppleARKitSettings>()->ARKitTimecodeProvider;
-		UTimecodeProvider* TimecodeProvider = FindObject<UTimecodeProvider>(GEngine, *ProviderName);
-		if (TimecodeProvider == nullptr)
-		{
-			// Try to load the class that was requested
-			UClass* Class = LoadClass<UTimecodeProvider>(nullptr, *ProviderName);
-			if (Class != nullptr)
-			{
-				TimecodeProvider = NewObject<UTimecodeProvider>(GEngine, Class);
-			}
-		}
-		// Create the default one if this failed for some reason
-		if (TimecodeProvider == nullptr)
-		{
-			TimecodeProvider = NewObject<UTimecodeProvider>(GEngine, UAppleARKitTimecodeProvider::StaticClass());
-		}
-		return TimecodeProvider;
-	}
+	static UTimecodeProvider* GetTimecodeProvider();
 };
