Fixed a race condition where the AR overlay would try to render before it had camera images

diff --git a/Engine/Plugins/Runtime/AR/Apple/AppleARKit/Source/AppleARKit/Private/AppleARKitVideoOverlay.cpp b/Engine/Plugins/Runtime/AR/Apple/AppleARKit/Source/AppleARKit/Private/AppleARKitVideoOverlay.cpp
--- a/Engine/Plugins/Runtime/AR/Apple/AppleARKit/Source/AppleARKit/Private/AppleARKitVideoOverlay.cpp
+++ b/Engine/Plugins/Runtime/AR/Apple/AppleARKit/Source/AppleARKit/Private/AppleARKitVideoOverlay.cpp
@@ -377,7 +377,7 @@ void FAppleARKitVideoOverlay::RenderVideoOverlay_RenderThread(FRHICommandListImm
 
 	if (FAppleARKitAvailability::SupportsARKit10())
 	{
-		if (RenderingOverlayMaterial == nullptr || !RenderingOverlayMaterial->IsValidLowLevel())
+		if (VideoTextureY == nullptr || VideoTextureCbCr == nullptr || RenderingOverlayMaterial == nullptr || !RenderingOverlayMaterial->IsValidLowLevel())
 		{
 			return;
 		}
