Reworked face ar console command support
Commands are now thread safe
Commands are:

ARKitSettings StartFileWriting
ARKitSettings StopFileWriting
ARKitSettings StartCameraFileWriting
ARKitSettings StopCameraFileWriting
ARKitSettings SavePerFrame
ARKitSettings SaveOnDemand
ARKitSettings FaceLogDir=<dir>
ARKitSettings LiveLinkSubjectName=<name>

diff --git a/Engine/Plugins/Runtime/AR/Apple/AppleARKit/Source/AppleARKit/Private/AppleARKitSystem.cpp b/Engine/Plugins/Runtime/AR/Apple/AppleARKit/Source/AppleARKit/Private/AppleARKitSystem.cpp
--- a/Engine/Plugins/Runtime/AR/Apple/AppleARKit/Source/AppleARKit/Private/AppleARKitSystem.cpp
+++ b/Engine/Plugins/Runtime/AR/Apple/AppleARKit/Source/AppleARKit/Private/AppleARKitSystem.cpp
@@ -209,6 +209,9 @@ FAppleARKitSystem::FAppleARKitSystem()
 , TimecodeProvider(nullptr)
 {
 	// See Initialize(), as we need access to SharedThis()
+#if SUPPORTS_ARKIT_1_0
+	IAppleImageUtilsPlugin::Load();
+#endif
 }
 
 FAppleARKitSystem::~FAppleARKitSystem()
@@ -1266,12 +1269,6 @@ bool FAppleARKitSystem::Run(UARSessionConfig* SessionConfig)
 	}
 
 #if SUPPORTS_ARKIT_1_0
-	// Set this based upon the project settings
-	bShouldWriteCameraImagePerFrame = GetDefault<UAppleARKitSettings>()->bShouldWriteCameraImagePerFrame;
-	WrittenCameraImageScale = GetDefault<UAppleARKitSettings>()->WrittenCameraImageScale;
-	WrittenCameraImageRotation = GetDefault<UAppleARKitSettings>()->WrittenCameraImageRotation;
-	WrittenCameraImageQuality = GetDefault<UAppleARKitSettings>()->WrittenCameraImageQuality;
-
 	if (FAppleARKitAvailability::SupportsARKit10())
 	{
 		ARSessionRunOptions options = 0;
@@ -1338,10 +1335,10 @@ bool FAppleARKitSystem::Run(UARSessionConfig* SessionConfig)
 		
 #if PLATFORM_IOS && !PLATFORM_TVOS
 		// Check if we need to adjust the priorities to allow ARKit to have more CPU time
-		if (GetDefault<UAppleARKitSettings>()->bAdjustThreadPrioritiesDuringARSession)
+		if (GetMutableDefault<UAppleARKitSettings>()->ShouldAdjustThreadPriorities())
 		{
-			int32 GameOverride = GetDefault<UAppleARKitSettings>()->GameThreadPriorityOverride;
-			int32 RenderOverride = GetDefault<UAppleARKitSettings>()->RenderThreadPriorityOverride;
+			int32 GameOverride = GetMutableDefault<UAppleARKitSettings>()->GetGameThreadPriorityOverride();
+			int32 RenderOverride = GetMutableDefault<UAppleARKitSettings>()->GetRenderThreadPriorityOverride();
 			SetThreadPriority(GameOverride);
 			if (XRCamera.IsValid())
 			{
@@ -1405,7 +1402,7 @@ bool FAppleARKitSystem::Pause()
 	
 #if PLATFORM_IOS && !PLATFORM_TVOS
 	// Check if we need to adjust the priorities to allow ARKit to have more CPU time
-	if (GetDefault<UAppleARKitSettings>()->bAdjustThreadPrioritiesDuringARSession)
+	if (GetMutableDefault<UAppleARKitSettings>()->ShouldAdjustThreadPriorities())
 	{
 		SetThreadPriority(GAME_THREAD_PRIORITY);
 		if (XRCamera.IsValid())
@@ -1441,7 +1438,7 @@ void FAppleARKitSystem::SessionDidUpdateFrame_DelegateThread(TSharedPtr< FAppleA
 	{
 		UpdateARKitPerfStats();
 #if SUPPORTS_ARKIT_1_0
-		if (bShouldWriteCameraImagePerFrame)
+		if (GetMutableDefault<UAppleARKitSettings>()->ShouldWriteCameraImagePerFrame())
 		{
 			WriteCameraImageToDisk(Frame->CameraImage);
 		}
@@ -1955,9 +1952,9 @@ void FAppleARKitSystem::UpdateARKitPerfStats()
 #if SUPPORTS_ARKIT_1_0
 void FAppleARKitSystem::WriteCameraImageToDisk(CVPixelBufferRef PixelBuffer)
 {
-	int32 ImageQuality = WrittenCameraImageQuality;
-	float ImageScale = WrittenCameraImageScale;
-	ETextureRotationDirection ImageRotation = WrittenCameraImageRotation;
+	int32 ImageQuality = GetMutableDefault<UAppleARKitSettings>()->GetWrittenCameraImageQuality();
+	float ImageScale = GetMutableDefault<UAppleARKitSettings>()->GetWrittenCameraImageScale();
+	ETextureRotationDirection ImageRotation = GetMutableDefault<UAppleARKitSettings>()->GetWrittenCameraImageRotation();
 	CIImage* SourceImage = [[CIImage alloc] initWithCVPixelBuffer: PixelBuffer];
 	FTimecode Timecode = TimecodeProvider->GetTimecode();
 	AsyncTask(ENamedThreads::AnyBackgroundThreadNormalTask, [SourceImage, ImageQuality, ImageScale, ImageRotation, Timecode]()
@@ -1968,7 +1965,9 @@ void FAppleARKitSystem::WriteCameraImageToDisk(CVPixelBufferRef PixelBuffer)
 		// Build a unique file name
 		FDateTime DateTime = FDateTime::UtcNow();
 		static FString UserDir = FPlatformProcess::UserDir();
-		FString FileName = FString::Printf(TEXT("%sCameraImages/Image_%d-%d-%d-%d-%d-%d-%d.jpeg"), *UserDir,
+		const FString& FaceDir = GetMutableDefault<UAppleARKitSettings>()->GetFaceTrackingLogDir();
+		const TCHAR* SubDir = FaceDir.Len() > 0 ? *FaceDir : TEXT("CameraImages");
+		FString FileName = FString::Printf(TEXT("%s%s/Image_%d-%d-%d-%d-%d-%d-%d.jpeg"), *UserDir, SubDir,
 			DateTime.GetYear(), DateTime.GetMonth(), DateTime.GetDay(), Timecode.Hours, Timecode.Minutes, Timecode.Seconds, Timecode.Frames);
 		// Write the jpeg to disk
 		FFileHelper::SaveArrayToFile(JpegBytes, *FileName);
@@ -2014,6 +2013,170 @@ UTimecodeProvider* UAppleARKitSettings::GetTimecodeProvider()
 	return TimecodeProvider;
 }
 
+void UAppleARKitSettings::CreateFaceTrackingLogDir()
+{
+	const FString& FaceDir = GetMutableDefault<UAppleARKitSettings>()->GetFaceTrackingLogDir();
+	const TCHAR* SubDir = FaceDir.Len() > 0 ? *FaceDir : TEXT("FaceTracking");
+	const FString UserDir = FPlatformProcess::UserDir();
+	if (!IFileManager::Get().DirectoryExists(*(UserDir / SubDir)))
+	{
+		IFileManager::Get().MakeDirectory(*(UserDir / SubDir));
+	}
+}
+
+void UAppleARKitSettings::CreateImageLogDir()
+{
+	const FString& FaceDir = GetMutableDefault<UAppleARKitSettings>()->GetFaceTrackingLogDir();
+	const TCHAR* SubDir = FaceDir.Len() > 0 ? *FaceDir : TEXT("CameraImages");
+	const FString UserDir = FPlatformProcess::UserDir();
+	if (!IFileManager::Get().DirectoryExists(*(UserDir / SubDir)))
+	{
+		IFileManager::Get().MakeDirectory(*(UserDir / SubDir));
+	}
+}
+
+FString UAppleARKitSettings::GetFaceTrackingLogDir()
+{
+	FScopeLock ScopeLock(&CriticalSection);
+	return FaceTrackingLogDir;
+}
+
+bool UAppleARKitSettings::IsLiveLinkEnabledForFaceTracking()
+{
+	FScopeLock ScopeLock(&CriticalSection);
+	return bEnableLiveLinkForFaceTracking;
+}
+
+bool UAppleARKitSettings::IsFaceTrackingLoggingEnabled()
+{
+	FScopeLock ScopeLock(&CriticalSection);
+	return bFaceTrackingLogData;
+}
+
+bool UAppleARKitSettings::ShouldFaceTrackingLogPerFrame()
+{
+	FScopeLock ScopeLock(&CriticalSection);
+	return bFaceTrackingWriteEachFrame;
+}
+
+EARFaceTrackingFileWriterType UAppleARKitSettings::GetFaceTrackingFileWriterType()
+{
+	FScopeLock ScopeLock(&CriticalSection);
+	return FaceTrackingFileWriterType;
+}
+
+bool UAppleARKitSettings::ShouldWriteCameraImagePerFrame()
+{
+	FScopeLock ScopeLock(&CriticalSection);
+	return bShouldWriteCameraImagePerFrame;
+}
+
+float UAppleARKitSettings::GetWrittenCameraImageScale()
+{
+	FScopeLock ScopeLock(&CriticalSection);
+	return WrittenCameraImageScale;
+}
+
+int32 UAppleARKitSettings::GetWrittenCameraImageQuality()
+{
+	FScopeLock ScopeLock(&CriticalSection);
+	return WrittenCameraImageQuality;
+}
+
+ETextureRotationDirection UAppleARKitSettings::GetWrittenCameraImageRotation()
+{
+	FScopeLock ScopeLock(&CriticalSection);
+	return WrittenCameraImageRotation;
+}
+
+int32 UAppleARKitSettings::GetLiveLinkPublishingPort()
+{
+	FScopeLock ScopeLock(&CriticalSection);
+	return LiveLinkPublishingPort;
+}
+
+FName UAppleARKitSettings::GetLiveLinkSubjectName()
+{
+	FScopeLock ScopeLock(&CriticalSection);
+	return DefaultFaceTrackingLiveLinkSubjectName;
+}
+
+EARFaceTrackingDirection UAppleARKitSettings::GetFaceTrackingDirection()
+{
+	FScopeLock ScopeLock(&CriticalSection);
+	return DefaultFaceTrackingDirection;
+}
+
+bool UAppleARKitSettings::ShouldAdjustThreadPriorities()
+{
+	FScopeLock ScopeLock(&CriticalSection);
+	return bAdjustThreadPrioritiesDuringARSession;
+}
+
+int32 UAppleARKitSettings::GetGameThreadPriorityOverride()
+{
+	FScopeLock ScopeLock(&CriticalSection);
+	return GameThreadPriorityOverride;
+}
+
+int32 UAppleARKitSettings::GetRenderThreadPriorityOverride()
+{
+	FScopeLock ScopeLock(&CriticalSection);
+	return RenderThreadPriorityOverride;
+}
+
+bool UAppleARKitSettings::Exec(UWorld*, const TCHAR* Cmd, FOutputDevice& Ar)
+{
+	if (FParse::Command(&Cmd, TEXT("ARKitSettings")))
+	{
+		FScopeLock ScopeLock(&CriticalSection);
+
+		if (FParse::Command(&Cmd, TEXT("StartFileWriting")))
+		{
+			UAppleARKitSettings::CreateFaceTrackingLogDir();
+			bFaceTrackingLogData = true;
+			bShouldWriteCameraImagePerFrame = true;
+			return true;
+		}
+		else if (FParse::Command(&Cmd, TEXT("StopFileWriting")))
+		{
+			bFaceTrackingLogData = false;
+			bShouldWriteCameraImagePerFrame = false;
+			return true;
+		}
+		else if (FParse::Command(&Cmd, TEXT("StartCameraFileWriting")))
+		{
+			bShouldWriteCameraImagePerFrame = true;
+			return true;
+		}
+		else if (FParse::Command(&Cmd, TEXT("StopCameraFileWriting")))
+		{
+			bShouldWriteCameraImagePerFrame = false;
+			return true;
+		}
+		else if (FParse::Command(&Cmd, TEXT("SavePerFrame")))
+		{
+			bFaceTrackingWriteEachFrame = true;
+			return true;
+		}
+		else if (FParse::Command(&Cmd, TEXT("SaveOnDemand")))
+		{
+			bFaceTrackingWriteEachFrame = false;
+			return true;
+		}
+		else if (FParse::Value(Cmd, TEXT("FaceLogDir="), FaceTrackingLogDir))
+		{
+			UAppleARKitSettings::CreateFaceTrackingLogDir();
+			return true;
+		}
+		else if (FParse::Value(Cmd, TEXT("LiveLinkSubjectName="), DefaultFaceTrackingLiveLinkSubjectName))
+		{
+			return true;
+		}
+	}
+	return false;
+}
+
 #if PLATFORM_IOS
 	#pragma clang diagnostic pop
 #endif
diff --git a/Engine/Plugins/Runtime/AR/Apple/AppleARKit/Source/AppleARKit/Private/AppleARKitSystem.h b/Engine/Plugins/Runtime/AR/Apple/AppleARKit/Source/AppleARKit/Private/AppleARKitSystem.h
--- a/Engine/Plugins/Runtime/AR/Apple/AppleARKit/Source/AppleARKit/Private/AppleARKitSystem.h
+++ b/Engine/Plugins/Runtime/AR/Apple/AppleARKit/Source/AppleARKit/Private/AppleARKitSystem.h
@@ -98,10 +98,6 @@ private:
 	void UpdateFrame();
 	void CalcTrackingToWorldRotation();
 #if SUPPORTS_ARKIT_1_0
-	bool bShouldWriteCameraImagePerFrame;
-	float WrittenCameraImageScale;
-	int32 WrittenCameraImageQuality;
-	ETextureRotationDirection WrittenCameraImageRotation;
 	/** Asynchronously writes a JPEG to disk */
 	void WriteCameraImageToDisk(CVPixelBufferRef PixelBuffer);
 #endif
diff --git a/Engine/Plugins/Runtime/AR/Apple/AppleARKit/Source/AppleARKit/Public/AppleARKitSettings.h b/Engine/Plugins/Runtime/AR/Apple/AppleARKit/Source/AppleARKit/Public/AppleARKitSettings.h
--- a/Engine/Plugins/Runtime/AR/Apple/AppleARKit/Source/AppleARKit/Public/AppleARKitSettings.h
+++ b/Engine/Plugins/Runtime/AR/Apple/AppleARKit/Source/AppleARKit/Public/AppleARKitSettings.h
@@ -26,7 +26,8 @@ enum class EARFaceTrackingFileWriterType : uint8
 
 UCLASS(Config=Engine, defaultconfig)
 class APPLEARKIT_API UAppleARKitSettings :
-	public UObject
+	public UObject,
+	public FSelfRegisteringExec
 {
 	GENERATED_BODY()
 
@@ -49,68 +50,101 @@ public:
 	{
 	}
 
+public:
+	static UTimecodeProvider* GetTimecodeProvider();
+	static void CreateFaceTrackingLogDir();
+	static void CreateImageLogDir();
+
+// Accessors to the properties
+	FString GetFaceTrackingLogDir();
+	bool IsLiveLinkEnabledForFaceTracking();
+	bool IsFaceTrackingLoggingEnabled();
+	bool ShouldFaceTrackingLogPerFrame();
+	EARFaceTrackingFileWriterType GetFaceTrackingFileWriterType();
+	bool ShouldWriteCameraImagePerFrame();
+	float GetWrittenCameraImageScale();
+	int32 GetWrittenCameraImageQuality();
+	ETextureRotationDirection GetWrittenCameraImageRotation();
+	int32 GetLiveLinkPublishingPort();
+	FName GetLiveLinkSubjectName();
+	EARFaceTrackingDirection GetFaceTrackingDirection();
+	bool ShouldAdjustThreadPriorities();
+	int32 GetGameThreadPriorityOverride();
+	int32 GetRenderThreadPriorityOverride();
+
+protected:
+	//~ FSelfRegisteringExec
+	virtual bool Exec(UWorld*, const TCHAR* Cmd, FOutputDevice& Ar) override;
+	//~ FSelfRegisteringExec
+
+	/** Where to write the curve files and image files */
+	UPROPERTY(Config, BlueprintReadOnly, EditAnywhere, Category="AR Settings")
+	FString FaceTrackingLogDir;
+
 	/** Whether to publish face blend shapes to LiveLink or not */
-	UPROPERTY(Config, EditAnywhere, Category="AR Settings")
+	UPROPERTY(Config, BlueprintReadOnly, EditAnywhere, Category="AR Settings")
 	bool bEnableLiveLinkForFaceTracking;
 
 	/** Whether file writing is enabled at all or not */
-	UPROPERTY(Config, EditAnywhere, Category="AR Settings")
+	UPROPERTY(Config, BlueprintReadOnly, EditAnywhere, Category="AR Settings")
 	bool bFaceTrackingLogData;
 
 	/** Whether to publish each frame or when the "FaceAR WriteCurveFile */
-	UPROPERTY(Config, EditAnywhere, Category="AR Settings")
+	UPROPERTY(Config, BlueprintReadOnly, EditAnywhere, Category="AR Settings")
 	bool bFaceTrackingWriteEachFrame;
 
 	/** The type of face AR publisher that writes to disk to create */
-	UPROPERTY(Config, EditAnywhere, Category="AR Settings")
+	UPROPERTY(Config, BlueprintReadOnly, EditAnywhere, Category="AR Settings")
 	EARFaceTrackingFileWriterType FaceTrackingFileWriterType;
 
 	/** Whether to publish the camera image each frame */
-	UPROPERTY(Config, EditAnywhere, Category="AR Settings")
+	UPROPERTY(Config, BlueprintReadOnly, EditAnywhere, Category="AR Settings")
 	bool bShouldWriteCameraImagePerFrame;
 
 	/** The scale to write the images at. Used to reduce data footprint */
-	UPROPERTY(Config, EditAnywhere, Category="AR Settings")
+	UPROPERTY(Config, BlueprintReadOnly, EditAnywhere, Category="AR Settings")
 	float WrittenCameraImageScale;
 
 	/** The quality setting to generate the jpeg images at. Defaults to 85, which is "high quality". Lower values reduce data footprint */
-	UPROPERTY(Config, EditAnywhere, Category="AR Settings")
+	UPROPERTY(Config, BlueprintReadOnly, EditAnywhere, Category="AR Settings")
 	int32 WrittenCameraImageQuality;
 
 	/** Defaults to none. Use Right when in portrait mode */
-	UPROPERTY(Config, EditAnywhere, Category="AR Settings")
+	UPROPERTY(Config, BlueprintReadOnly, EditAnywhere, Category="AR Settings")
 	ETextureRotationDirection WrittenCameraImageRotation;
 
 	/** The port to use when listening/sending LiveLink face blend shapes via the network */
-	UPROPERTY(Config, EditAnywhere, Category="AR Settings")
+	UPROPERTY(Config, BlueprintReadOnly, EditAnywhere, Category="AR Settings")
 	int32 LiveLinkPublishingPort;
 
 	/** The default name to use when publishing face tracking name */
-	UPROPERTY(Config, EditAnywhere, Category="AR Settings")
+	UPROPERTY(Config, BlueprintReadOnly, EditAnywhere, Category="AR Settings")
 	FName DefaultFaceTrackingLiveLinkSubjectName;
 
 	/** The default tracking to use when tracking face blend shapes (face relative or mirrored). Defaults to face relative */
-	UPROPERTY(Config, EditAnywhere, Category="AR Settings")
+	UPROPERTY(Config, BlueprintReadOnly, EditAnywhere, Category="AR Settings")
 	EARFaceTrackingDirection DefaultFaceTrackingDirection;
 
 	/** Whether to adjust thread priorities during an AR session or not */
-	UPROPERTY(Config, EditAnywhere, Category="AR Settings")
+	UPROPERTY(Config, BlueprintReadOnly, EditAnywhere, Category="AR Settings")
 	bool bAdjustThreadPrioritiesDuringARSession;
 
 	/** The game thread priority to change to when an AR session is running, default is 47 */
-	UPROPERTY(Config, EditAnywhere, Category="AR Settings")
+	UPROPERTY(Config, BlueprintReadOnly, EditAnywhere, Category="AR Settings")
 	int32 GameThreadPriorityOverride;
 	
 	/** The render thread priority to change to when an AR session is running, default is 45 */
-	UPROPERTY(Config, EditAnywhere, Category="AR Settings")
+	UPROPERTY(Config, BlueprintReadOnly, EditAnywhere, Category="AR Settings")
 	int32 RenderThreadPriorityOverride;
 
 	/**
 	 * Used to specify the timecode provider to use when identifying when an update occurred.
 	 * Useful when using external timecode generators to sync multiple devices/machines
 	 */
-	UPROPERTY(Config, EditAnywhere, Category="AR Settings")
+	UPROPERTY(Config, BlueprintReadOnly, EditAnywhere, Category="AR Settings")
 	FString ARKitTimecodeProvider;
 
-	static UTimecodeProvider* GetTimecodeProvider();
+private:
+	/** Used because these defaults need to be read on multiple threads */
+	FCriticalSection CriticalSection;
 };
diff --git a/Engine/Plugins/Runtime/AR/Apple/AppleARKitFaceSupport/Source/AppleARKitFaceSupport/Private/AppleARKitFaceMeshComponent.cpp b/Engine/Plugins/Runtime/AR/Apple/AppleARKitFaceSupport/Source/AppleARKitFaceSupport/Private/AppleARKitFaceMeshComponent.cpp
--- a/Engine/Plugins/Runtime/AR/Apple/AppleARKitFaceSupport/Source/AppleARKitFaceSupport/Private/AppleARKitFaceMeshComponent.cpp
+++ b/Engine/Plugins/Runtime/AR/Apple/AppleARKitFaceSupport/Source/AppleARKitFaceSupport/Private/AppleARKitFaceMeshComponent.cpp
@@ -31,7 +31,7 @@ NSDictionary<ARBlendShapeLocation,NSNumber *>* ToBlendShapeDictionary(const FARB
 	}
 
 	// Do we want to capture face performance or look at the face as if in a mirror (Apple is mirrored so we mirror the mirror)
-	if (GetDefault<UAppleARKitSettings>()->DefaultFaceTrackingDirection == EARFaceTrackingDirection::FaceMirrored)
+	if (GetMutableDefault<UAppleARKitSettings>()->GetFaceTrackingDirection() == EARFaceTrackingDirection::FaceMirrored)
 	{
 		SET_BLEND_SHAPE(ARBlendShapeLocationEyeBlinkLeft, EARFaceBlendShape::EyeBlinkLeft);
 		SET_BLEND_SHAPE(ARBlendShapeLocationEyeLookDownLeft, EARFaceBlendShape::EyeLookDownLeft);
diff --git a/Engine/Plugins/Runtime/AR/Apple/AppleARKitFaceSupport/Source/AppleARKitFaceSupport/Private/AppleARKitFaceSupportImpl.cpp b/Engine/Plugins/Runtime/AR/Apple/AppleARKitFaceSupport/Source/AppleARKitFaceSupport/Private/AppleARKitFaceSupportImpl.cpp
--- a/Engine/Plugins/Runtime/AR/Apple/AppleARKitFaceSupport/Source/AppleARKitFaceSupport/Private/AppleARKitFaceSupportImpl.cpp
+++ b/Engine/Plugins/Runtime/AR/Apple/AppleARKitFaceSupport/Source/AppleARKitFaceSupport/Private/AppleARKitFaceSupportImpl.cpp
@@ -214,9 +214,8 @@ void FAppleARKitFaceSupport::PublishLiveLinkData(TSharedPtr<FAppleARKitAnchorDat
 	{
 		bNeedsInit = false;
 	    // Create our LiveLink provider if the project setting is enabled
-		if (GetDefault<UAppleARKitSettings>()->bEnableLiveLinkForFaceTracking)
+		if (GetMutableDefault<UAppleARKitSettings>()->IsLiveLinkEnabledForFaceTracking())
 		{
-			FaceTrackingLiveLinkSubjectName = GetDefault<UAppleARKitSettings>()->DefaultFaceTrackingLiveLinkSubjectName;
 #if PLATFORM_IOS
 			LiveLinkSource = FAppleARKitLiveLinkSourceFactory::CreateLiveLinkSource();
 #else
@@ -228,6 +227,7 @@ void FAppleARKitFaceSupport::PublishLiveLinkData(TSharedPtr<FAppleARKitAnchorDat
 
 	if (LiveLinkSource.IsValid())
 	{
+		FaceTrackingLiveLinkSubjectName = GetMutableDefault<UAppleARKitSettings>()->GetLiveLinkSubjectName();
         LiveLinkSource->PublishBlendShapes(FaceTrackingLiveLinkSubjectName, Anchor->Timecode, Anchor->FrameRate, Anchor->BlendShapes, LocalDeviceId);
 	}
 }
diff --git a/Engine/Plugins/Runtime/AR/Apple/AppleARKitFaceSupport/Source/AppleARKitFaceSupport/Private/AppleARKitLiveLinkSource.cpp b/Engine/Plugins/Runtime/AR/Apple/AppleARKitFaceSupport/Source/AppleARKitFaceSupport/Private/AppleARKitLiveLinkSource.cpp
--- a/Engine/Plugins/Runtime/AR/Apple/AppleARKitFaceSupport/Source/AppleARKitFaceSupport/Private/AppleARKitLiveLinkSource.cpp
+++ b/Engine/Plugins/Runtime/AR/Apple/AppleARKitFaceSupport/Source/AppleARKitFaceSupport/Private/AppleARKitLiveLinkSource.cpp
@@ -268,7 +268,7 @@ TSharedRef<FInternetAddr> FAppleARKitLiveLinkRemotePublisher::GetSendAddress()
 	// Don't bother trying to parse the IP if it isn't set
 	if (RemoteIp.Len())
 	{
-		int32 LiveLinkPort = GetDefault<UAppleARKitSettings>()->LiveLinkPublishingPort;
+		int32 LiveLinkPort = GetMutableDefault<UAppleARKitSettings>()->GetLiveLinkPublishingPort();
 		SendAddr->SetPort(LiveLinkPort);
 		bool bIsValid = false;
 		SendAddr->SetIp(*RemoteIp, bIsValid);
@@ -437,17 +437,14 @@ void FAppleARKitLiveLinkRemoteListener::Tick(float DeltaTime)
 
 FAppleARKitLiveLinkFileWriter::FAppleARKitLiveLinkFileWriter(const TCHAR* InFileExtension)
 	: FileExtension(InFileExtension)
-	, bSavePerFrameOrOnDemand(false)
 {
-	// Read the config values for this
-	GConfig->GetBool(TEXT("/Script/AppleARKit.AppleARKitSettings"), TEXT("bFaceTrackingWriteEachFrame"), bSavePerFrameOrOnDemand, GEngineIni);
-	GConfig->GetBool(TEXT("/Script/AppleARKit.AppleARKitSettings"), TEXT("bFaceTrackingLogData"), bLogData, GEngineIni);
+	UAppleARKitSettings::CreateFaceTrackingLogDir();
 }
 
 FAppleARKitLiveLinkFileWriter::~FAppleARKitLiveLinkFileWriter()
 {
 	// Save on close if desired
-	if (!bSavePerFrameOrOnDemand)
+	if (!GetMutableDefault<UAppleARKitSettings>()->ShouldFaceTrackingLogPerFrame())
 	{
 		SaveFileData();
 	}
@@ -475,7 +472,9 @@ FString FAppleARKitLiveLinkFileWriter::GenerateFilePath()
 	FDateTime DateTime = FDateTime::UtcNow();
 	const FString UserDir = FPlatformProcess::UserDir();
 	const FString DeviceNameString = DeviceName.ToString();
-	return FString::Printf(TEXT("%sFaceTracking/%s_%d-%d-%d-%d-%d-%d-%d%s"), *UserDir, *DeviceNameString,
+	const FString& FaceDir = GetMutableDefault<UAppleARKitSettings>()->GetFaceTrackingLogDir();
+	const TCHAR* SubDir = FaceDir.Len() > 0 ? *FaceDir : TEXT("FaceTracking");
+	return FString::Printf(TEXT("%s%s/%s_%d-%d-%d-%d-%d-%d-%d%s"), *UserDir, SubDir, *DeviceNameString,
 		DateTime.GetYear(), DateTime.GetMonth(), DateTime.GetDay(), Timecode.Hours, Timecode.Minutes, Timecode.Seconds, Timecode.Frames,
 		*FileExtension);
 }
@@ -484,7 +483,7 @@ void FAppleARKitLiveLinkFileWriter::PublishBlendShapes(FName SubjectName, const
 {
 	FScopeLock ScopeLock(&CriticalSection);
 
-	if (!bLogData)
+	if (!GetMutableDefault<UAppleARKitSettings>()->IsFaceTrackingLoggingEnabled())
 	{
 		return;
 	}
@@ -493,7 +492,7 @@ void FAppleARKitLiveLinkFileWriter::PublishBlendShapes(FName SubjectName, const
 	// Add to the array for long running save
 	new(FrameHistory) FFaceTrackingFrame(Timecode, FrameRate, FaceBlendShapes);
 
-	if (bSavePerFrameOrOnDemand)
+	if (GetMutableDefault<UAppleARKitSettings>()->ShouldFaceTrackingLogPerFrame())
 	{
 		SaveFileData();
 	}
@@ -509,32 +508,6 @@ bool FAppleARKitLiveLinkFileWriter::Exec(UWorld*, const TCHAR* Cmd, FOutputDevic
 			SaveFileData();
 			return true;
 		}
-		else if (FParse::Command(&Cmd, TEXT("StartFileWriting")))
-		{
-			FScopeLock ScopeLock(&CriticalSection);
-			FrameHistory.Empty();
-			bLogData = true;
-			return true;
-		}
-		else if (FParse::Command(&Cmd, TEXT("StopFileWriting")))
-		{
-			FScopeLock ScopeLock(&CriticalSection);
-			FrameHistory.Empty();
-			bLogData = false;
-			return true;
-		}
-		else if (FParse::Command(&Cmd, TEXT("SavePerFrame")))
-		{
-			FScopeLock ScopeLock(&CriticalSection);
-			bSavePerFrameOrOnDemand = true;
-			return true;
-		}
-		else if (FParse::Command(&Cmd, TEXT("SaveOnDemand")))
-		{
-			FScopeLock ScopeLock(&CriticalSection);
-			bSavePerFrameOrOnDemand = false;
-			return true;
-		}
 	}
 	return false;
 }
@@ -563,8 +536,8 @@ FAppleARKitLiveLinkFileWriterCsv::FAppleARKitLiveLinkFileWriterCsv()
 FString FAppleARKitLiveLinkFileWriterCsv::BuildCsvRow(const FFaceTrackingFrame& Frame)
 {
 	FString SaveData = FString::Printf(TEXT("%d:%d:%d:%d, %d"),
-			Frame.Timecode.Hours, Frame.Timecode.Minutes, Frame.Timecode.Seconds, Frame.Timecode.Frames,
-			Frame.FrameRate);
+		Frame.Timecode.Hours, Frame.Timecode.Minutes, Frame.Timecode.Seconds, Frame.Timecode.Frames,
+		Frame.FrameRate);
 	// Add all of the blend shapes on
 	for (int32 Shape = 0; Shape < (int32)EARFaceBlendShape::MAX; Shape++)
 	{
diff --git a/Engine/Plugins/Runtime/AR/Apple/AppleARKitFaceSupport/Source/AppleARKitFaceSupport/Private/AppleARKitLiveLinkSource.h b/Engine/Plugins/Runtime/AR/Apple/AppleARKitFaceSupport/Source/AppleARKitFaceSupport/Private/AppleARKitLiveLinkSource.h
--- a/Engine/Plugins/Runtime/AR/Apple/AppleARKitFaceSupport/Source/AppleARKitFaceSupport/Private/AppleARKitLiveLinkSource.h
+++ b/Engine/Plugins/Runtime/AR/Apple/AppleARKitFaceSupport/Source/AppleARKitFaceSupport/Private/AppleARKitLiveLinkSource.h
@@ -825,10 +825,6 @@ protected:
 	TArray<FFaceTrackingFrame> FrameHistory;
 	/** The name to use when generating the file to save to */
 	FName DeviceName;
-	/** If false, the calls to PublishCurves are ignored */
-	bool bLogData;
-	/** If true, it writes a file each frame. Otherwise, it writes the data upon demand or exit */
-	bool bSavePerFrameOrOnDemand;
 	/**
 	 * The time code provider to use when tagging file time stamps
 	 * Note: this requires the FAppleARKitSystem object to mark it in use so GC doesn't destroy it. Normally it would
