Fix for race condition of images during camera image logging

diff --git a/Engine/Plugins/Runtime/AR/Apple/AppleARKit/Source/AppleARKit/Private/AppleARKitSystem.cpp b/Engine/Plugins/Runtime/AR/Apple/AppleARKit/Source/AppleARKit/Private/AppleARKitSystem.cpp
--- a/Engine/Plugins/Runtime/AR/Apple/AppleARKit/Source/AppleARKit/Private/AppleARKitSystem.cpp
+++ b/Engine/Plugins/Runtime/AR/Apple/AppleARKit/Source/AppleARKit/Private/AppleARKitSystem.cpp
@@ -1952,13 +1952,14 @@ void FAppleARKitSystem::UpdateARKitPerfStats()
 #if SUPPORTS_ARKIT_1_0
 void FAppleARKitSystem::WriteCameraImageToDisk(CVPixelBufferRef PixelBuffer)
 {
+	CFRetain(PixelBuffer);
 	int32 ImageQuality = GetMutableDefault<UAppleARKitSettings>()->GetWrittenCameraImageQuality();
 	float ImageScale = GetMutableDefault<UAppleARKitSettings>()->GetWrittenCameraImageScale();
 	ETextureRotationDirection ImageRotation = GetMutableDefault<UAppleARKitSettings>()->GetWrittenCameraImageRotation();
-	CIImage* SourceImage = [[CIImage alloc] initWithCVPixelBuffer: PixelBuffer];
 	FTimecode Timecode = TimecodeProvider->GetTimecode();
-	AsyncTask(ENamedThreads::AnyBackgroundThreadNormalTask, [SourceImage, ImageQuality, ImageScale, ImageRotation, Timecode]()
+	AsyncTask(ENamedThreads::AnyBackgroundThreadNormalTask, [PixelBuffer, ImageQuality, ImageScale, ImageRotation, Timecode]()
 	{
+		CIImage* SourceImage = [[CIImage alloc] initWithCVPixelBuffer: PixelBuffer];
 		TArray<uint8> JpegBytes;
 		IAppleImageUtilsPlugin::Get().ConvertToJPEG(SourceImage, JpegBytes, ImageQuality, true, true, ImageScale, ImageRotation);
 		[SourceImage release];
@@ -1971,6 +1972,7 @@ void FAppleARKitSystem::WriteCameraImageToDisk(CVPixelBufferRef PixelBuffer)
 			DateTime.GetYear(), DateTime.GetMonth(), DateTime.GetDay(), Timecode.Hours, Timecode.Minutes, Timecode.Seconds, Timecode.Frames);
 		// Write the jpeg to disk
 		FFileHelper::SaveArrayToFile(JpegBytes, *FileName);
+		CFRelease(PixelBuffer);
 	});
 }
 #endif
