Implement support for GetPassthroughCameraUVs_RenderThread on Apple ARKit.

diff --git a/Engine/Plugins/Runtime/AR/Apple/AppleARKit/Source/AppleARKit/Private/AppleARKitSystem.cpp b/Engine/Plugins/Runtime/AR/Apple/AppleARKit/Source/AppleARKit/Private/AppleARKitSystem.cpp
--- a/Engine/Plugins/Runtime/AR/Apple/AppleARKit/Source/AppleARKit/Private/AppleARKitSystem.cpp
+++ b/Engine/Plugins/Runtime/AR/Apple/AppleARKit/Source/AppleARKit/Private/AppleARKitSystem.cpp
@@ -157,6 +157,11 @@ private:
 		VideoOverlay.RenderVideoOverlay_RenderThread(RHICmdList, InView, ARKitSystem.DeviceOrientation);
 	}
 	
+	virtual bool GetPassthroughCameraUVs_RenderThread(TArray<FVector2D>& OutUVs) override
+	{
+		return VideoOverlay.GetPassthroughCameraUVs_RenderThread(OutUVs, ARKitSystem.DeviceOrientation);
+	}
+
 	virtual bool IsActiveThisFrame(class FViewport* InViewport) const override
 	{
 		// Base implementation needs this call as it updates bCurrentFrameIsStereoRendering as a side effect.
diff --git a/Engine/Plugins/Runtime/AR/Apple/AppleARKit/Source/AppleARKit/Private/AppleARKitVideoOverlay.cpp b/Engine/Plugins/Runtime/AR/Apple/AppleARKit/Source/AppleARKit/Private/AppleARKitVideoOverlay.cpp
--- a/Engine/Plugins/Runtime/AR/Apple/AppleARKit/Source/AppleARKit/Private/AppleARKitVideoOverlay.cpp
+++ b/Engine/Plugins/Runtime/AR/Apple/AppleARKit/Source/AppleARKit/Private/AppleARKitVideoOverlay.cpp
@@ -150,7 +150,7 @@ void FAppleARKitVideoOverlay::UpdateVideoTexture_RenderThread(FRHICommandListImm
 			}
 		}
 		
-		const FVector2D UVOffset = (ViewAspectRatioLandscape <= Frame.Camera.GetAspectRatio()) ? FVector2D(UVOffsetAmount, 0.0f) : FVector2D(0.0f, UVOffsetAmount);
+		UVOffset = (ViewAspectRatioLandscape <= Frame.Camera.GetAspectRatio()) ? FVector2D(UVOffsetAmount, 0.0f) : FVector2D(0.0f, UVOffsetAmount);
 		
 		// Setup vertex buffer
 		const FVector4 Positions[] =
@@ -481,3 +481,53 @@ void FAppleARKitVideoOverlay::RenderVideoOverlay_RenderThread(FRHICommandListImm
 	}
 #endif
 }
+
+
+bool FAppleARKitVideoOverlay::GetPassthroughCameraUVs_RenderThread(TArray<FVector2D>& OutUVs, const EScreenOrientation::Type DeviceOrientation)
+{
+#if SUPPORTS_ARKIT_1_0
+	if (VideoTextureY != nullptr) // make sure UVOffset has been initialized
+	{
+		OutUVs.SetNumUninitialized(4);
+
+		switch (DeviceOrientation)
+		{
+		case EScreenOrientation::Type::LandscapeLeft:
+			OutUVs[1] = FVector2D(UVOffset.X, 1.0f - UVOffset.Y);
+			OutUVs[0] = FVector2D(UVOffset.X, UVOffset.Y);
+			OutUVs[3] = FVector2D(1.0f - UVOffset.X, 1.0f - UVOffset.Y);
+			OutUVs[2] = FVector2D(1.0f - UVOffset.X, UVOffset.Y);
+			return true;
+
+		case EScreenOrientation::Type::LandscapeRight:
+			OutUVs[1] = FVector2D(1.0f - UVOffset.X, UVOffset.Y);
+			OutUVs[0] = FVector2D(1.0f - UVOffset.X, 1.0f - UVOffset.Y);
+			OutUVs[3] = FVector2D(UVOffset.X, UVOffset.Y);
+			OutUVs[2] = FVector2D(UVOffset.X, 1.0f - UVOffset.Y);
+			return true;
+
+		case EScreenOrientation::Type::Portrait:
+			OutUVs[1] = FVector2D(1.0f - UVOffset.X, 1.0f - UVOffset.Y);
+			OutUVs[0] = FVector2D(UVOffset.X, 1.0f - UVOffset.Y);
+			OutUVs[3] = FVector2D(1.0f - UVOffset.X, UVOffset.Y);
+			OutUVs[2] = FVector2D(UVOffset.X, UVOffset.Y);
+			return true;
+
+		case EScreenOrientation::Type::PortraitUpsideDown:
+			OutUVs[1] = FVector2D(UVOffset.X, UVOffset.Y);
+			OutUVs[0] = FVector2D(1.0f - UVOffset.X, UVOffset.Y);
+			OutUVs[3] = FVector2D(UVOffset.X, 1.0f - UVOffset.Y);
+			OutUVs[2] = FVector2D(1.0f - UVOffset.X, 1.0f - UVOffset.Y);
+			return true;
+
+		default:
+			return false;
+		}
+	}
+	else
+#endif
+	{
+		return false;
+	}
+}
+
diff --git a/Engine/Plugins/Runtime/AR/Apple/AppleARKit/Source/AppleARKit/Private/AppleARKitVideoOverlay.h b/Engine/Plugins/Runtime/AR/Apple/AppleARKit/Source/AppleARKit/Private/AppleARKitVideoOverlay.h
--- a/Engine/Plugins/Runtime/AR/Apple/AppleARKit/Source/AppleARKit/Private/AppleARKitVideoOverlay.h
+++ b/Engine/Plugins/Runtime/AR/Apple/AppleARKit/Source/AppleARKit/Private/AppleARKitVideoOverlay.h
@@ -36,12 +36,16 @@ public:
 
 	void UpdateVideoTexture_RenderThread(FRHICommandListImmediate& RHICmdList, FAppleARKitFrame& Frame, const FSceneViewFamily& InViewFamily);
 	void RenderVideoOverlay_RenderThread(FRHICommandListImmediate& RHICmdList, const FSceneView& InView, const EScreenOrientation::Type DeviceOrientation);
+	bool GetPassthroughCameraUVs_RenderThread(TArray<FVector2D>& OutUVs, const EScreenOrientation::Type DeviceOrientation);
 
 private:
 	FTextureRHIRef VideoTextureY;
 	FTextureRHIRef VideoTextureCbCr;
 	UMaterialInterface* RenderingOverlayMaterial;
 	FIndexBufferRHIRef OverlayIndexBufferRHI;
+
+	// Cache UV offset to be used by GetPassthroughCameraUVs_RenderThread
+	FVector2D UVOffset;
 	
 	// Separate vertex buffer for each supported device orientation
 	FVertexBufferRHIRef OverlayVertexBufferRHI[4];
