Added a log message on failure to write the image to disk when logging images

diff --git a/Engine/Plugins/Runtime/AR/Apple/AppleARKit/Source/AppleARKit/Private/AppleARKitSystem.cpp b/Engine/Plugins/Runtime/AR/Apple/AppleARKit/Source/AppleARKit/Private/AppleARKitSystem.cpp
--- a/Engine/Plugins/Runtime/AR/Apple/AppleARKit/Source/AppleARKit/Private/AppleARKitSystem.cpp
+++ b/Engine/Plugins/Runtime/AR/Apple/AppleARKit/Source/AppleARKit/Private/AppleARKitSystem.cpp
@@ -1976,7 +1976,10 @@ void FAppleARKitSystem::WriteCameraImageToDisk(CVPixelBufferRef PixelBuffer)
 		FString FileName = FString::Printf(TEXT("%s%s/Image_%d-%d-%d-%d-%d-%d-%d.jpeg"), *UserDir, SubDir,
 			DateTime.GetYear(), DateTime.GetMonth(), DateTime.GetDay(), Timecode.Hours, Timecode.Minutes, Timecode.Seconds, Timecode.Frames);
 		// Write the jpeg to disk
-		FFileHelper::SaveArrayToFile(JpegBytes, *FileName);
+		if (!FFileHelper::SaveArrayToFile(JpegBytes, *FileName))
+		{
+			UE_LOG(LogAppleARKit, Error, TEXT("Failed to save JPEG to file name '%s'"), *FileName);
+		}
 		CFRelease(PixelBuffer);
 	});
 }
