Merged in UDN fix for orientation only ar sessions not enabling the auto focus flag if requested

diff --git a/Engine/Plugins/Runtime/AR/Apple/AppleARKit/Source/AppleARKit/Private/AppleARKitConversion.cpp b/Engine/Plugins/Runtime/AR/Apple/AppleARKit/Source/AppleARKit/Private/AppleARKitConversion.cpp
--- a/Engine/Plugins/Runtime/AR/Apple/AppleARKit/Source/AppleARKit/Private/AppleARKitConversion.cpp
+++ b/Engine/Plugins/Runtime/AR/Apple/AppleARKit/Source/AppleARKit/Private/AppleARKitConversion.cpp
@@ -244,7 +244,14 @@ ARConfiguration* FAppleARKitConversion::ToARConfiguration( UARSessionConfig* Ses
 			{
 				return nullptr;
 			}
-			SessionConfiguration = [AROrientationTrackingConfiguration new];
+			AROrientationTrackingConfiguration* OrientationTrackingConfiguration = [AROrientationTrackingConfiguration new];
+#if SUPPORTS_ARKIT_1_5
+			if (FAppleARKitAvailability::SupportsARKit15())
+			{
+				OrientationTrackingConfiguration.autoFocusEnabled = SessionConfig->ShouldEnableAutoFocus();
+			}
+#endif
+			SessionConfiguration = OrientationTrackingConfiguration;
 			break;
 		}
 		case EARSessionType::World:
