Added MRMesh generation (collision, nav mesh, and rendering) for AR planes and AR tracked images on ARKit

diff --git a/Engine/Plugins/Runtime/AR/Apple/AppleARKit/Source/AppleARKit/AppleARKit.Build.cs b/Engine/Plugins/Runtime/AR/Apple/AppleARKit/Source/AppleARKit/AppleARKit.Build.cs
--- a/Engine/Plugins/Runtime/AR/Apple/AppleARKit/Source/AppleARKit/AppleARKit.Build.cs
+++ b/Engine/Plugins/Runtime/AR/Apple/AppleARKit/Source/AppleARKit/AppleARKit.Build.cs
@@ -36,7 +36,8 @@ public class AppleARKit : ModuleRules
                 "RenderCore",
                 "HeadMountedDisplay",
                 "AugmentedReality",
-                "AppleImageUtils"
+                "AppleImageUtils",
+                "MRMesh"
 				// ... add private dependencies that you statically link with here ...
 			}
 			);
diff --git a/Engine/Plugins/Runtime/AR/Apple/AppleARKit/Source/AppleARKit/Private/AppleARKitSystem.cpp b/Engine/Plugins/Runtime/AR/Apple/AppleARKit/Source/AppleARKit/Private/AppleARKitSystem.cpp
--- a/Engine/Plugins/Runtime/AR/Apple/AppleARKit/Source/AppleARKit/Private/AppleARKitSystem.cpp
+++ b/Engine/Plugins/Runtime/AR/Apple/AppleARKit/Source/AppleARKit/Private/AppleARKitSystem.cpp
@@ -22,6 +22,10 @@
 #include "HAL/ThreadSafeCounter.h"
 #include "Misc/FileHelper.h"
 
+// For mesh occlusion
+#include "MRMeshComponent.h"
+#include "AROriginActor.h"
+
 // To separate out the face ar library linkage from standard ar apps
 #include "AppleARKitFaceSupport.h"
 
@@ -1273,7 +1277,11 @@ bool FAppleARKitSystem::Run(UARSessionConfig* SessionConfig)
 		SetDeviceOrientation( ScreenOrientation );
 	}
 
+
 #if SUPPORTS_ARKIT_1_0
+	// Don't do the conversion work if they don't want this
+	FAppleARKitAnchorData::bGenerateGeometry = SessionConfig->bGenerateMeshDataFromTrackedGeometry;
+
 	if (FAppleARKitAvailability::SupportsARKit10())
 	{
 		ARSessionRunOptions options = 0;
@@ -1459,6 +1467,7 @@ void FAppleARKitSystem::SessionDidFailWithError_DelegateThread(const FString& Er
 #if SUPPORTS_ARKIT_1_0
 
 TArray<int32> FAppleARKitAnchorData::FaceIndices;
+bool FAppleARKitAnchorData::bGenerateGeometry = false;
 
 static TSharedPtr<FAppleARKitAnchorData> MakeAnchorData( ARAnchor* Anchor, double Timestamp, uint32 FrameNumber )
 {
@@ -1471,19 +1480,37 @@ static TSharedPtr<FAppleARKitAnchorData> MakeAnchorData( ARAnchor* Anchor, doubl
 			FAppleARKitConversion::ToFTransform(PlaneAnchor.transform),
 			FAppleARKitConversion::ToFVector(PlaneAnchor.center),
 			// @todo use World Settings WorldToMetersScale
-			0.5f*FAppleARKitConversion::ToFVector(PlaneAnchor.extent).GetAbs()
+			0.5f*FAppleARKitConversion::ToFVector(PlaneAnchor.extent).GetAbs(),
+			FAppleARKitConversion::ToEARPlaneOrientation(PlaneAnchor.alignment)
 		);
 
 #if SUPPORTS_ARKIT_1_5
 		if (FAppleARKitAvailability::SupportsARKit15())
 		{
-			//@todo All this copying should really happen on-demand.
-			const int32 NumBoundaryVerts = PlaneAnchor.geometry.boundaryVertexCount;
-			NewAnchor->BoundaryVerts.Reset(NumBoundaryVerts);
-			for(int32 i=0; i<NumBoundaryVerts; ++i)
+			if (FAppleARKitAnchorData::bGenerateGeometry)
 			{
-				const vector_float3& Vert = PlaneAnchor.geometry.boundaryVertices[i];
-				NewAnchor->BoundaryVerts.Add(FAppleARKitConversion::ToFVector(Vert));
+				const int32 NumBoundaryVerts = PlaneAnchor.geometry.boundaryVertexCount;
+				NewAnchor->BoundaryVerts.Reset(NumBoundaryVerts);
+				for(int32 i=0; i<NumBoundaryVerts; ++i)
+				{
+					const vector_float3& Vert = PlaneAnchor.geometry.boundaryVertices[i];
+					NewAnchor->BoundaryVerts.Add(FAppleARKitConversion::ToFVector(Vert));
+				}
+				// Generate the mesh from the plane
+				NewAnchor->Vertices.Reset(4);
+				NewAnchor->Vertices.Add(NewAnchor->Center + NewAnchor->Extent);
+				NewAnchor->Vertices.Add(NewAnchor->Center + FVector(NewAnchor->Extent.X, -NewAnchor->Extent.Y, NewAnchor->Extent.Z));
+				NewAnchor->Vertices.Add(NewAnchor->Center + FVector(-NewAnchor->Extent.X, -NewAnchor->Extent.Y, NewAnchor->Extent.Z));
+				NewAnchor->Vertices.Add(NewAnchor->Center + FVector(-NewAnchor->Extent.X, NewAnchor->Extent.Y, NewAnchor->Extent.Z));
+
+				// Two triangles
+				NewAnchor->Indices.Reset(6);
+				NewAnchor->Indices.Add(0);
+				NewAnchor->Indices.Add(1);
+				NewAnchor->Indices.Add(2);
+				NewAnchor->Indices.Add(2);
+				NewAnchor->Indices.Add(3);
+				NewAnchor->Indices.Add(0);
 			}
 		}
 #endif
@@ -1503,6 +1530,26 @@ static TSharedPtr<FAppleARKitAnchorData> MakeAnchorData( ARAnchor* Anchor, doubl
 		{
 			NewAnchor->bIsTracked = ImageAnchor.isTracked;
 		}
+		if (FAppleARKitAnchorData::bGenerateGeometry)
+		{
+			FVector Extent(ImageAnchor.referenceImage.physicalSize.width, ImageAnchor.referenceImage.physicalSize.height, 0.f);
+			Extent *= 100.f;
+			// Generate the mesh from the reference image's sizes
+			NewAnchor->Vertices.Reset(4);
+			NewAnchor->Vertices.Add(Extent);
+			NewAnchor->Vertices.Add(FVector(Extent.X, -Extent.Y, Extent.Z));
+			NewAnchor->Vertices.Add(FVector(-Extent.X, -Extent.Y, Extent.Z));
+			NewAnchor->Vertices.Add(FVector(-Extent.X, Extent.Y, Extent.Z));
+			
+			// Two triangles
+			NewAnchor->Indices.Reset(6);
+			NewAnchor->Indices.Add(0);
+			NewAnchor->Indices.Add(1);
+			NewAnchor->Indices.Add(2);
+			NewAnchor->Indices.Add(2);
+			NewAnchor->Indices.Add(3);
+			NewAnchor->Indices.Add(0);
+		}
 #endif
 	}
 #endif
@@ -1654,6 +1701,30 @@ void FAppleARKitSystem::SessionDidAddAnchors_Internal( TSharedRef<FAppleARKitAnc
 			NewAnchorDebugName = FString::Printf(TEXT("PLN-%02d"), LastTrackedGeometry_DebugId++);
 			UARPlaneGeometry* NewGeo = NewObject<UARPlaneGeometry>();
 			NewGeo->UpdateTrackedGeometry(ARComponent.ToSharedRef(), AnchorData->FrameNumber, AnchorData->Timestamp, AnchorData->Transform, GetARCompositionComponent()->GetAlignmentTransform(), AnchorData->Center, AnchorData->Extent);
+			NewGeo->SetOrientation(AnchorData->Orientation);
+			const UARSessionConfig& SessionConfig = GetARCompositionComponent()->GetSessionConfig();
+			// Add the occlusion geo if configured
+			if (SessionConfig.bGenerateMeshDataFromTrackedGeometry)
+			{
+				AAROriginActor* OriginActor = AAROriginActor::GetOriginActor();
+				UMRMeshComponent* MRMesh = NewObject<UMRMeshComponent>(OriginActor);
+
+				// Set the occlusion and wireframe defaults
+				MRMesh->SetEnableMeshOcclusion(SessionConfig.bUseMeshDataForOcclusion);
+				MRMesh->SetUseWireframe(SessionConfig.bRenderMeshDataInWireframe);
+				MRMesh->SetNeverCreateCollisionMesh(!SessionConfig.bGenerateCollisionForMeshData);
+				MRMesh->SetEnableNavMesh(SessionConfig.bGenerateNavMeshForMeshData);
+
+				// Set parent and register
+				MRMesh->SetupAttachment(OriginActor->GetRootComponent());
+				MRMesh->RegisterComponent();
+
+				// MRMesh takes ownership of the data in the arrays at this point
+				MRMesh->UpdateMesh(AnchorData->Transform.GetLocation(), AnchorData->Transform.GetRotation(), AnchorData->Transform.GetScale3D(), AnchorData->Vertices, AnchorData->Indices);
+
+				// Connect the tracked geo to the MRMesh
+				NewGeo->SetUnderlyingMesh(MRMesh);
+			}
 			NewGeometry = NewGeo;
 			break;
 		}
@@ -1679,6 +1750,29 @@ PRAGMA_ENABLE_DEPRECATION_WARNINGS
 			FVector2D PhysicalSize((*CandidateImage)->GetPhysicalWidth(), (*CandidateImage)->GetPhysicalHeight());
 			NewImage->UpdateTrackedGeometry(ARComponent.ToSharedRef(), AnchorData->FrameNumber, AnchorData->Timestamp, AnchorData->Transform, GetARCompositionComponent()->GetAlignmentTransform(), PhysicalSize, *CandidateImage);
 			NewGeometry = NewImage;
+			const UARSessionConfig& SessionConfig = GetARCompositionComponent()->GetSessionConfig();
+			// Add the occlusion geo if configured
+			if (SessionConfig.bGenerateMeshDataFromTrackedGeometry)
+			{
+				AAROriginActor* OriginActor = AAROriginActor::GetOriginActor();
+				UMRMeshComponent* MRMesh = NewObject<UMRMeshComponent>(OriginActor);
+				
+				// Set the occlusion and wireframe defaults
+				MRMesh->SetEnableMeshOcclusion(SessionConfig.bUseMeshDataForOcclusion);
+				MRMesh->SetUseWireframe(SessionConfig.bRenderMeshDataInWireframe);
+				MRMesh->SetNeverCreateCollisionMesh(!SessionConfig.bGenerateCollisionForMeshData);
+				MRMesh->SetEnableNavMesh(SessionConfig.bGenerateNavMeshForMeshData);
+				
+				// Set parent and register
+				MRMesh->SetupAttachment(OriginActor->GetRootComponent());
+				MRMesh->RegisterComponent();
+				
+				// MRMesh takes ownership of the data in the arrays at this point
+				MRMesh->UpdateMesh(AnchorData->Transform.GetLocation(), AnchorData->Transform.GetRotation(), AnchorData->Transform.GetScale3D(), AnchorData->Vertices, AnchorData->Indices);
+				
+				// Connect the tracked geo to the MRMesh
+				NewImage->SetUnderlyingMesh(MRMesh);
+			}
 			break;
 		}
 		case EAppleAnchorType::EnvironmentProbeAnchor:
@@ -1766,6 +1860,15 @@ void FAppleARKitSystem::SessionDidUpdateAnchors_Internal( TSharedRef<FAppleARKit
 						const FTransform Pin_LocalToTrackingTransform_PostUpdate = Pin->GetLocalToTrackingTransform_NoAlignment() * AnchorDeltaTransform;
 						Pin->OnTransformUpdated(Pin_LocalToTrackingTransform_PostUpdate);
 					}
+					PlaneGeo->SetOrientation(AnchorData->Orientation);
+					// Update the occlusion geo if configured
+					if (GetARCompositionComponent()->GetSessionConfig().bGenerateMeshDataFromTrackedGeometry)
+					{
+						UMRMeshComponent* MRMesh = PlaneGeo->GetUnderlyingMesh();
+						check(MRMesh != nullptr);
+						// MRMesh takes ownership of the data in the arrays at this point
+						MRMesh->UpdateMesh(AnchorData->Transform.GetLocation(), AnchorData->Transform.GetRotation(), AnchorData->Transform.GetScale3D(), AnchorData->Vertices, AnchorData->Indices);
+					}
 				}
 				break;
 			}
@@ -1805,6 +1908,14 @@ PRAGMA_ENABLE_DEPRECATION_WARNINGS
 						const FTransform Pin_LocalToTrackingTransform_PostUpdate = Pin->GetLocalToTrackingTransform_NoAlignment() * AnchorDeltaTransform;
 						Pin->OnTransformUpdated(Pin_LocalToTrackingTransform_PostUpdate);
 					}
+					// Update the occlusion geo if configured
+					if (GetARCompositionComponent()->GetSessionConfig().bGenerateMeshDataFromTrackedGeometry)
+					{
+						UMRMeshComponent* MRMesh = ImageAnchor->GetUnderlyingMesh();
+						check(MRMesh != nullptr);
+						// MRMesh takes ownership of the data in the arrays at this point
+						MRMesh->UpdateMesh(AnchorData->Transform.GetLocation(), AnchorData->Transform.GetRotation(), AnchorData->Transform.GetScale3D(), AnchorData->Vertices, AnchorData->Indices);
+					}
 				}
                 break;
             }
@@ -1842,7 +1953,14 @@ void FAppleARKitSystem::SessionDidRemoveAnchors_Internal( FGuid AnchorGuid )
 	{
 		UARTrackedGeometry* TrackedGeometryBeingRemoved = TrackedGeometries.FindChecked(AnchorGuid);
 		TrackedGeometryBeingRemoved->UpdateTrackingState(EARTrackingState::StoppedTracking);
-		
+		// Remove the occlusion mesh if present
+		UMRMeshComponent* MRMesh = TrackedGeometryBeingRemoved->GetUnderlyingMesh();
+		if (MRMesh != nullptr)
+		{
+			MRMesh->UnregisterComponent();
+			TrackedGeometryBeingRemoved->SetUnderlyingMesh(nullptr);
+		}
+
 		TArray<UARPin*> ARPinsBeingOrphaned = ARKitUtil::PinsFromGeometry(TrackedGeometryBeingRemoved, Pins);
 		for(UARPin* PinBeingOrphaned : ARPinsBeingOrphaned)
 		{
@@ -2187,6 +2305,124 @@ bool UAppleARKitSettings::Exec(UWorld*, const TCHAR* Cmd, FOutputDevice& Ar)
 	return false;
 }
 
+
+/** Used to run Exec commands */
+static bool MeshARTestingExec(UWorld* InWorld, const TCHAR* Cmd, FOutputDevice& Ar)
+{
+	bool bHandled = false;
+
+	if (FParse::Command(&Cmd, TEXT("ARKIT")))
+	{
+		if (FParse::Command(&Cmd, TEXT("MRMESH")))
+		{
+			AAROriginActor* OriginActor = AAROriginActor::GetOriginActor();
+			UMRMeshComponent* NewComp = NewObject<UMRMeshComponent>(OriginActor);
+			NewComp->RegisterComponent();
+			NewComp->SetUseWireframe(true);
+			// Send a fake update to it
+			FTransform Transform = FTransform::Identity;
+			static TArray<FVector> Verts;
+			static TArray<MRMESH_INDEX_TYPE> Indices;
+
+			Verts.Reset();
+			Indices.Reset();
+
+
+			Verts.Add(FVector(61.504616, 41.864193, 0.000000));
+			Verts.Add(FVector(62.357094, 40.308838, 0.000000));
+			Verts.Add(FVector(62.357094, -10.522773, 0.000000));
+			Verts.Add(FVector(51.219948, -21.031740, 0.000000));
+			Verts.Add(FVector(47.228138, -24.083036, 0.000000));
+			Verts.Add(FVector(42.975994, -26.713457, 0.000000));
+			Verts.Add(FVector(38.820320, -28.767712, 0.000000));
+			Verts.Add(FVector(-11.059804, -33.275490, 0.000000));
+			Verts.Add(FVector(-14.477231, -32.972424, 0.000000));
+			Verts.Add(FVector(-17.509483, -30.532234, 0.000000));
+			Verts.Add(FVector(-28.095594, -19.442810, 0.000000));
+			Verts.Add(FVector(-30.726019, -15.190664, 0.000000));
+			Verts.Add(FVector(-30.915573, 2.330182, 0.000000));
+			Verts.Add(FVector(-30.476627, 6.487241, 0.000000));
+			Verts.Add(FVector(-8.244256, 37.077099, 0.000000));
+			Verts.Add(FVector(-3.992107, 39.707523, 0.000000));
+			Verts.Add(FVector(12.504792, 46.403351, 0.000000));
+			Verts.Add(FVector(59.446259, 45.191578, 0.000000));
+
+
+			Indices.Add(7);
+			Indices.Add(8);
+			Indices.Add(9);
+
+			Indices.Add(10);
+			Indices.Add(11);
+			Indices.Add(12);
+
+			Indices.Add(7);
+			Indices.Add(9);
+			Indices.Add(10);
+
+			Indices.Add(7);
+			Indices.Add(10);
+			Indices.Add(12);
+
+			Indices.Add(7);
+			Indices.Add(12);
+			Indices.Add(13);
+
+			Indices.Add(7);
+			Indices.Add(13);
+			Indices.Add(14);
+
+			Indices.Add(6);
+			Indices.Add(7);
+			Indices.Add(14);
+
+			Indices.Add(6);
+			Indices.Add(14);
+			Indices.Add(15);
+
+			Indices.Add(6);
+			Indices.Add(15);
+			Indices.Add(16);
+
+			Indices.Add(5);
+			Indices.Add(6);
+			Indices.Add(16);
+
+			Indices.Add(4);
+			Indices.Add(5);
+			Indices.Add(16);
+
+			Indices.Add(3);
+			Indices.Add(4);
+			Indices.Add(16);
+
+			Indices.Add(2);
+			Indices.Add(3);
+			Indices.Add(16);
+
+			Indices.Add(1);
+			Indices.Add(2);
+			Indices.Add(16);
+
+			Indices.Add(0);
+			Indices.Add(1);
+			Indices.Add(16);
+
+			Indices.Add(0);
+			Indices.Add(16);
+			Indices.Add(17);
+
+			NewComp->UpdateMesh(Transform.GetLocation(), Transform.GetRotation(), Transform.GetScale3D(), Verts, Indices);
+
+			return true;
+		}
+	}
+
+	return false;
+}
+
+FStaticSelfRegisteringExec MeshARTestingExecRegistration(MeshARTestingExec);
+
 #if PLATFORM_IOS
 	#pragma clang diagnostic pop
 #endif
diff --git a/Engine/Plugins/Runtime/AR/Apple/AppleARKit/Source/AppleARKit/Public/AppleARKitConversion.h b/Engine/Plugins/Runtime/AR/Apple/AppleARKit/Source/AppleARKit/Public/AppleARKitConversion.h
--- a/Engine/Plugins/Runtime/AR/Apple/AppleARKit/Source/AppleARKit/Public/AppleARKitConversion.h
+++ b/Engine/Plugins/Runtime/AR/Apple/AppleARKit/Source/AppleARKit/Public/AppleARKitConversion.h
@@ -171,6 +171,19 @@ struct APPLEARKIT_API FAppleARKitConversion
 
 	static ARWorldAlignment ToARWorldAlignment( const EARWorldAlignment& InWorldAlignment );
 
+	/**
+	 * Coverts plane orientation
+	 */
+	static FORCEINLINE EARPlaneOrientation ToEARPlaneOrientation(ARPlaneAnchorAlignment Alignment)
+	{
+		EARPlaneOrientation RetVal = EARPlaneOrientation::Horizontal;
+		if (Alignment == ARPlaneAnchorAlignmentVertical)
+		{
+			RetVal = EARPlaneOrientation::Vertical;
+		}
+		return RetVal;
+	}
+
 #pragma clang diagnostic push
 #pragma clang diagnostic ignored "-Wpartial-availability"
 
@@ -240,12 +253,13 @@ struct FAppleARKitAnchorData
 	{
 	}
 
-	FAppleARKitAnchorData(FGuid InAnchorGuid, FTransform InTransform, FVector InCenter, FVector InExtent)
+	FAppleARKitAnchorData(FGuid InAnchorGuid, FTransform InTransform, FVector InCenter, FVector InExtent, EARPlaneOrientation InOrientation)
 		: Transform( InTransform )
 		, AnchorType( EAppleAnchorType::PlaneAnchor )
 		, AnchorGUID( InAnchorGuid )
 		, Center(InCenter)
 		, Extent(InExtent)
+		, Orientation(InOrientation)
 	{
 	}
 
@@ -331,7 +345,12 @@ struct FAppleARKitAnchorData
 	FGuid AnchorGUID;
 	FVector Center;
 	FVector Extent;
+	EARPlaneOrientation Orientation;
+	/** Set by the session config to detemine whether to generate geometry or not */
+	static bool bGenerateGeometry;
 	TArray<FVector> BoundaryVerts;
+	TArray<FVector> Vertices;
+	TArray<uint16> Indices;
 
 	FARBlendShapeMap BlendShapes;
 	TArray<FVector> FaceVerts;
diff --git a/Engine/Source/Runtime/AugmentedReality/Public/ARTrackable.h b/Engine/Source/Runtime/AugmentedReality/Public/ARTrackable.h
--- a/Engine/Source/Runtime/AugmentedReality/Public/ARTrackable.h
+++ b/Engine/Source/Runtime/AugmentedReality/Public/ARTrackable.h
@@ -123,7 +123,14 @@ public:
 	UFUNCTION(BlueprintPure, Category = "AR AugmentedReality|Plane Geometry")
 	UARPlaneGeometry* GetSubsumedBy() const { return SubsumedBy; };
 
+	UFUNCTION(BlueprintPure, Category="AR AugmentedReality|Plane Geometry")
+	EARPlaneOrientation GetOrientation() const { return Orientation; }
+	void SetOrientation(EARPlaneOrientation InOrientation) { Orientation = InOrientation; }
+
 private:
+	UPROPERTY()
+	EARPlaneOrientation Orientation;
+
 	UPROPERTY()
 	FVector Center;
 	
diff --git a/Engine/Source/Runtime/AugmentedReality/Public/ARTypes.h b/Engine/Source/Runtime/AugmentedReality/Public/ARTypes.h
--- a/Engine/Source/Runtime/AugmentedReality/Public/ARTypes.h
+++ b/Engine/Source/Runtime/AugmentedReality/Public/ARTypes.h
@@ -133,6 +133,13 @@ enum class EARWorldMappingState : uint8
 	Mapped
 };
 
+/** Describes the tracked plane orientation */
+UENUM(BlueprintType)
+enum class EARPlaneOrientation : uint8
+{
+	Horizontal,
+	Vertical
+};
 
 /** The current state of the AR subsystem including an optional explanation string. */
 USTRUCT(BlueprintType)
