Changed the debug exec for spawning a test mesh in ARKit
Fixed the wrong scale being used for an AR tracked image

diff --git a/Engine/Plugins/Runtime/AR/Apple/AppleARKit/Source/AppleARKit/Private/AppleARKitSystem.cpp b/Engine/Plugins/Runtime/AR/Apple/AppleARKit/Source/AppleARKit/Private/AppleARKitSystem.cpp
--- a/Engine/Plugins/Runtime/AR/Apple/AppleARKit/Source/AppleARKit/Private/AppleARKitSystem.cpp
+++ b/Engine/Plugins/Runtime/AR/Apple/AppleARKit/Source/AppleARKit/Private/AppleARKitSystem.cpp
@@ -1533,7 +1533,8 @@ static TSharedPtr<FAppleARKitAnchorData> MakeAnchorData( ARAnchor* Anchor, doubl
 		if (FAppleARKitAnchorData::bGenerateGeometry)
 		{
 			FVector Extent(ImageAnchor.referenceImage.physicalSize.width, ImageAnchor.referenceImage.physicalSize.height, 0.f);
-			Extent *= 100.f;
+			// Scale by half since this is an extent around the center (same as scale then divide by 2)
+			Extent *= 50.f;
 			// Generate the mesh from the reference image's sizes
 			NewAnchor->Vertices.Reset(4);
 			NewAnchor->Vertices.Add(Extent);
@@ -2321,98 +2322,24 @@ static bool MeshARTestingExec(UWorld* InWorld, const TCHAR* Cmd, FOutputDevice&
 			NewComp->SetUseWireframe(true);
 			// Send a fake update to it
 			FTransform Transform = FTransform::Identity;
-			static TArray<FVector> Verts;
-			static TArray<MRMESH_INDEX_TYPE> Indices;
-
-			Verts.Reset();
-			Indices.Reset();
-
-
-			Verts.Add(FVector(61.504616, 41.864193, 0.000000));
-			Verts.Add(FVector(62.357094, 40.308838, 0.000000));
-			Verts.Add(FVector(62.357094, -10.522773, 0.000000));
-			Verts.Add(FVector(51.219948, -21.031740, 0.000000));
-			Verts.Add(FVector(47.228138, -24.083036, 0.000000));
-			Verts.Add(FVector(42.975994, -26.713457, 0.000000));
-			Verts.Add(FVector(38.820320, -28.767712, 0.000000));
-			Verts.Add(FVector(-11.059804, -33.275490, 0.000000));
-			Verts.Add(FVector(-14.477231, -32.972424, 0.000000));
-			Verts.Add(FVector(-17.509483, -30.532234, 0.000000));
-			Verts.Add(FVector(-28.095594, -19.442810, 0.000000));
-			Verts.Add(FVector(-30.726019, -15.190664, 0.000000));
-			Verts.Add(FVector(-30.915573, 2.330182, 0.000000));
-			Verts.Add(FVector(-30.476627, 6.487241, 0.000000));
-			Verts.Add(FVector(-8.244256, 37.077099, 0.000000));
-			Verts.Add(FVector(-3.992107, 39.707523, 0.000000));
-			Verts.Add(FVector(12.504792, 46.403351, 0.000000));
-			Verts.Add(FVector(59.446259, 45.191578, 0.000000));
-
-
-			Indices.Add(7);
-			Indices.Add(8);
-			Indices.Add(9);
-
-			Indices.Add(10);
-			Indices.Add(11);
-			Indices.Add(12);
-
-			Indices.Add(7);
-			Indices.Add(9);
-			Indices.Add(10);
-
-			Indices.Add(7);
-			Indices.Add(10);
-			Indices.Add(12);
-
-			Indices.Add(7);
-			Indices.Add(12);
-			Indices.Add(13);
-
-			Indices.Add(7);
-			Indices.Add(13);
-			Indices.Add(14);
-
-			Indices.Add(6);
-			Indices.Add(7);
-			Indices.Add(14);
-
-			Indices.Add(6);
-			Indices.Add(14);
-			Indices.Add(15);
-
-			Indices.Add(6);
-			Indices.Add(15);
-			Indices.Add(16);
-
-			Indices.Add(5);
-			Indices.Add(6);
-			Indices.Add(16);
-
-			Indices.Add(4);
-			Indices.Add(5);
-			Indices.Add(16);
+			TArray<FVector> Vertices;
+			TArray<MRMESH_INDEX_TYPE> Indices;
 
-			Indices.Add(3);
-			Indices.Add(4);
-			Indices.Add(16);
-
-			Indices.Add(2);
-			Indices.Add(3);
-			Indices.Add(16);
-
-			Indices.Add(1);
-			Indices.Add(2);
-			Indices.Add(16);
+			Vertices.Reset(4);
+			Vertices.Add(FVector(100.f, 100.f, 0.f));
+			Vertices.Add(FVector(100.f, -100.f, 0.f));
+			Vertices.Add(FVector(-100.f, -100.f, 0.f));
+			Vertices.Add(FVector(-100.f, 100.f, 0.f));
 
+			Indices.Reset(6);
 			Indices.Add(0);
 			Indices.Add(1);
-			Indices.Add(16);
-
+			Indices.Add(2);
+			Indices.Add(2);
+			Indices.Add(3);
 			Indices.Add(0);
-			Indices.Add(16);
-			Indices.Add(17);
 
-			NewComp->UpdateMesh(Transform.GetLocation(), Transform.GetRotation(), Transform.GetScale3D(), Verts, Indices);
+			NewComp->UpdateMesh(Transform.GetLocation(), Transform.GetRotation(), Transform.GetScale3D(), Vertices, Indices);
 
 			return true;
 		}
