First pass UARCameraImage support for ARKit

diff --git a/Engine/Plugins/Runtime/AR/Apple/AppleARKit/Source/AppleARKit/Private/AppleARKitSystem.cpp b/Engine/Plugins/Runtime/AR/Apple/AppleARKit/Source/AppleARKit/Private/AppleARKitSystem.cpp
--- a/Engine/Plugins/Runtime/AR/Apple/AppleARKit/Source/AppleARKit/Private/AppleARKitSystem.cpp
+++ b/Engine/Plugins/Runtime/AR/Apple/AppleARKit/Source/AppleARKit/Private/AppleARKitSystem.cpp
@@ -1455,6 +1455,11 @@ void FAppleARKitSystem::SessionDidUpdateFrame_DelegateThread(TSharedPtr< FAppleA
 		{
 			WriteCameraImageToDisk(Frame->CameraImage);
 		}
+		// Queue an update on the render thread
+		if (CameraImage != nullptr)
+		{
+			CameraImage->Init_RenderThread(Frame->CameraImage);
+		}
 #endif
 	}
 }
diff --git a/Engine/Plugins/Runtime/AR/Apple/AppleARKit/Source/AppleARKit/Private/AppleARKitTextures.cpp b/Engine/Plugins/Runtime/AR/Apple/AppleARKit/Source/AppleARKit/Private/AppleARKitTextures.cpp
--- a/Engine/Plugins/Runtime/AR/Apple/AppleARKit/Source/AppleARKit/Private/AppleARKitTextures.cpp
+++ b/Engine/Plugins/Runtime/AR/Apple/AppleARKit/Source/AppleARKit/Private/AppleARKitTextures.cpp
@@ -10,6 +10,127 @@
 	#import <Metal/Metal.h>
 #endif
 
+/** Resource class to do all of the setup work on the render thread */
+class FARKitCameraImageResource :
+	public FTextureResource
+{
+public:
+	FARKitCameraImageResource(UAppleARKitTextureCameraImage* InOwner)
+		: LastFrameNumber(0)
+		, Owner(InOwner)
+	{
+		CameraImage = Owner->GetCameraImage();
+		if (CameraImage != nullptr)
+		{
+			CFRetain(CameraImage);
+		}
+	}
+
+	virtual ~FARKitCameraImageResource()
+	{
+	}
+
+	virtual void InitRHI() override
+	{
+		FSamplerStateInitializerRHI SamplerStateInitializer(SF_Bilinear, AM_Clamp, AM_Clamp, AM_Clamp);
+		SamplerStateRHI = RHICreateSamplerState(SamplerStateInitializer);
+
+		if (CameraImage != nullptr)
+		{
+			SCOPED_AUTORELEASE_POOL;
+			
+			Size.X = CVPixelBufferGetWidth(CameraImage);
+			Size.Y = CVPixelBufferGetHeight(CameraImage);
+
+			// Create the target texture that we'll update into
+			FRHIResourceCreateInfo CreateInfo;
+			DecodedTextureRef = RHICreateTexture2D(Size.X, Size.Y, PF_B8G8R8A8, 1, 1, TexCreate_Dynamic | TexCreate_ShaderResource | TexCreate_UAV, CreateInfo);
+
+			// Get the underlying metal texture so we can render to it
+			id<MTLTexture> UnderlyingMetalTexture = (id<MTLTexture>)DecodedTextureRef->GetNativeResource();
+
+			CIImage* Image = [[CIImage alloc] initWithCVPixelBuffer: CameraImage];
+			CIContext* Context = [CIContext context];
+			CGColorSpaceRef ColorSpaceRef = CGColorSpaceCreateDeviceRGB();
+			// Do the conversion on the GPU
+			[Context render: Image toMTLTexture: UnderlyingMetalTexture commandBuffer: nil bounds: Image.extent colorSpace: ColorSpaceRef];
+
+			// Now that the conversion is done, we can get rid of our refs
+			[Image release];
+			CGColorSpaceRelease(ColorSpaceRef);
+			CFRelease(CameraImage);
+			CameraImage = nullptr;
+		}
+		else
+		{
+			// Default to an empty 1x1 texture if we don't have a camera image
+			FRHIResourceCreateInfo CreateInfo;
+			Size.X = Size.Y = 1;
+			DecodedTextureRef = RHICreateTexture2D(Size.X, Size.Y, PF_B8G8R8A8, 1, 1, TexCreate_ShaderResource, CreateInfo);
+		}
+
+		TextureRHI = DecodedTextureRef;
+		TextureRHI->SetName(Owner->GetFName());
+		RHIBindDebugLabelName(TextureRHI, *Owner->GetName());
+		RHIUpdateTextureReference(Owner->TextureReference.TextureReferenceRHI, TextureRHI);
+	}
+
+	virtual void ReleaseRHI() override
+	{
+		RHIUpdateTextureReference(Owner->TextureReference.TextureReferenceRHI, FTextureRHIParamRef());
+		if (CameraImage != nullptr)
+		{
+			CFRelease(CameraImage);
+		}
+		CameraImage = nullptr;
+		DecodedTextureRef.SafeRelease();
+		FTextureResource::ReleaseRHI();
+	}
+
+	/** Returns the width of the texture in pixels. */
+	virtual uint32 GetSizeX() const override
+	{
+		return Size.X;
+	}
+
+	/** Returns the height of the texture in pixels. */
+	virtual uint32 GetSizeY() const override
+	{
+		return Size.Y;
+	}
+
+	/** Render thread update of the texture so we don't get 2 updates per frame on the render thread */
+	void Init_RenderThread(CVPixelBufferRef InCameraImage)
+	{
+		check(IsInRenderingThread());
+		check(InCameraImage != nullptr);
+
+		if (LastFrameNumber != GFrameNumber)
+		{
+			LastFrameNumber = GFrameNumber;
+			ReleaseRHI();
+			CameraImage = InCameraImage;
+			CFRetain(CameraImage);
+			InitRHI();
+		}
+	}
+
+private:
+	/** The size we get from the incoming camera image */
+	FIntPoint Size;
+
+#if PLATFORM_MAC || PLATFORM_IOS
+	/** The raw camera image from ARKit which is to be converted into a RGBA image */
+	CVPixelBufferRef CameraImage;
+#endif
+	/** The texture that we actually render with which is populated via the GPU conversion process */
+	FTexture2DRHIRef DecodedTextureRef;
+	/** The last frame we were updated on */
+	uint32 LastFrameNumber;
+
+	const UAppleARKitTextureCameraImage* Owner;
+};
+
 UAppleARKitTextureCameraImage::UAppleARKitTextureCameraImage(const FObjectInitializer& ObjectInitializer)
 	: Super(ObjectInitializer)
 #if PLATFORM_MAC || PLATFORM_IOS
@@ -21,7 +142,9 @@ UAppleARKitTextureCameraImage::UAppleARKitTextureCameraImage(const FObjectInitia
 
 FTextureResource* UAppleARKitTextureCameraImage::CreateResource()
 {
-	// @todo joeg -- hook this up for rendering
+#if PLATFORM_MAC || PLATFORM_IOS
+	return new FARKitCameraImageResource(this);
+#endif
 	return nullptr;
 }
 
@@ -37,10 +160,12 @@ void UAppleARKitTextureCameraImage::BeginDestroy()
 	Super::BeginDestroy();
 }
 
-#if SUPPORTS_ARKIT_1_0
+#if PLATFORM_MAC || PLATFORM_IOS
 
 void UAppleARKitTextureCameraImage::Init(float InTimestamp, CVPixelBufferRef InCameraImage)
 {
+	check(IsInGameThread());
+
 	// Handle the case where this UObject is being reused
 	if (CameraImage != nullptr)
 	{
@@ -56,9 +181,28 @@ void UAppleARKitTextureCameraImage::Init(float InTimestamp, CVPixelBufferRef InC
 		Size.X = CVPixelBufferGetWidth(CameraImage);
 		Size.Y = CVPixelBufferGetHeight(CameraImage);
 	}
-	//@todo joeg - Update the render resources
+
+	if (Resource == nullptr)
+	{
+		// Initial update. All others will be queued on the render thread
+		UpdateResource();
+	}
 }
 
+void UAppleARKitTextureCameraImage::Init_RenderThread(CVPixelBufferRef InCameraImage)
+{
+	if (Resource != nullptr)
+	{
+		CFRetain(InCameraImage);
+		FARKitCameraImageResource* ARKitResource = static_cast<FARKitCameraImageResource*>(Resource);
+		ENQUEUE_RENDER_COMMAND(Init_RenderThread)(
+			[ARKitResource, InCameraImage](FRHICommandListImmediate&)
+		{
+			ARKitResource->Init_RenderThread(InCameraImage);
+			CFRelease(InCameraImage);
+		});
+	}
+}
 #endif
 
 UAppleARKitTextureCameraDepth::UAppleARKitTextureCameraDepth(const FObjectInitializer& ObjectInitializer)
diff --git a/Engine/Plugins/Runtime/AR/Apple/AppleARKit/Source/AppleARKit/Public/AppleARKitTextures.h b/Engine/Plugins/Runtime/AR/Apple/AppleARKit/Source/AppleARKit/Public/AppleARKitTextures.h
--- a/Engine/Plugins/Runtime/AR/Apple/AppleARKit/Source/AppleARKit/Public/AppleARKitTextures.h
+++ b/Engine/Plugins/Runtime/AR/Apple/AppleARKit/Source/AppleARKit/Public/AppleARKitTextures.h
@@ -22,22 +22,24 @@ public:
 	// UTexture interface implementation
 	virtual void BeginDestroy() override;
 	virtual FTextureResource* CreateResource() override;
-	virtual EMaterialValueType GetMaterialType() const override { return MCT_TextureExternal; }
+	virtual EMaterialValueType GetMaterialType() const override { return MCT_Texture2D; }
 	virtual float GetSurfaceWidth() const override { return Size.X; }
 	virtual float GetSurfaceHeight() const override { return Size.Y; }
-	virtual FGuid GetExternalTextureGuid() const override { return ExternalTextureGuid; }
 	// End UTexture interface
 
-#if SUPPORTS_ARKIT_1_0
+#if PLATFORM_MAC || PLATFORM_IOS
 	/** Sets any initialization data */
 	virtual void Init(float InTimestamp, CVPixelBufferRef InCameraImage);
+
+	/** Queues the conversion on the render thread */
+	virtual void Init_RenderThread(CVPixelBufferRef InCameraImage);
 #endif
 
 	virtual EAppleTextureType GetTextureType() const override { return EAppleTextureType::PixelBuffer; }
 
 #if PLATFORM_MAC || PLATFORM_IOS
 	/** Returns the cached camera image. You must retain this if you hold onto it */
-	CVPixelBufferRef GetCameraImage() { return CameraImage; }
+	CVPixelBufferRef GetCameraImage() const { return CameraImage; }
 
 	// IAppleImageInterface interface implementation
 	virtual CVPixelBufferRef GetPixelBuffer() const override { return CameraImage; }
