Added missing #ifdefs around Apple specific code

diff --git a/Engine/Plugins/Runtime/AR/Apple/AppleARKit/Source/AppleARKit/Private/AppleARKitTextures.cpp b/Engine/Plugins/Runtime/AR/Apple/AppleARKit/Source/AppleARKit/Private/AppleARKitTextures.cpp
--- a/Engine/Plugins/Runtime/AR/Apple/AppleARKit/Source/AppleARKit/Private/AppleARKitTextures.cpp
+++ b/Engine/Plugins/Runtime/AR/Apple/AppleARKit/Source/AppleARKit/Private/AppleARKitTextures.cpp
@@ -19,11 +19,13 @@ public:
 		: LastFrameNumber(0)
 		, Owner(InOwner)
 	{
+#if PLATFORM_MAC || PLATFORM_IOS
 		CameraImage = Owner->GetCameraImage();
 		if (CameraImage != nullptr)
 		{
 			CFRetain(CameraImage);
 		}
+#endif
 	}
 
 	virtual ~FARKitCameraImageResource()
@@ -35,6 +37,7 @@ public:
 		FSamplerStateInitializerRHI SamplerStateInitializer(SF_Bilinear, AM_Clamp, AM_Clamp, AM_Clamp);
 		SamplerStateRHI = RHICreateSamplerState(SamplerStateInitializer);
 
+#if PLATFORM_MAC || PLATFORM_IOS
 		if (CameraImage != nullptr)
 		{
 			SCOPED_AUTORELEASE_POOL;
@@ -62,6 +65,7 @@ public:
 			CameraImage = nullptr;
 		}
 		else
+#endif
 		{
 			// Default to an empty 1x1 texture if we don't have a camera image
 			FRHIResourceCreateInfo CreateInfo;
@@ -78,11 +82,13 @@ public:
 	virtual void ReleaseRHI() override
 	{
 		RHIUpdateTextureReference(Owner->TextureReference.TextureReferenceRHI, FTextureRHIParamRef());
+#if PLATFORM_MAC || PLATFORM_IOS
 		if (CameraImage != nullptr)
 		{
 			CFRelease(CameraImage);
 		}
 		CameraImage = nullptr;
+#endif
 		DecodedTextureRef.SafeRelease();
 		FTextureResource::ReleaseRHI();
 	}
@@ -99,6 +105,7 @@ public:
 		return Size.Y;
 	}
 
+#if PLATFORM_MAC || PLATFORM_IOS
 	/** Render thread update of the texture so we don't get 2 updates per frame on the render thread */
 	void Init_RenderThread(CVPixelBufferRef InCameraImage)
 	{
@@ -114,6 +121,7 @@ public:
 			InitRHI();
 		}
 	}
+#endif
 
 private:
 	/** The size we get from the incoming camera image */
