ARKit camera image texture objects now take device orientation into account

diff --git a/Engine/Plugins/Runtime/AR/Apple/AppleARKit/Source/AppleARKit/Private/AppleARKitTextures.cpp b/Engine/Plugins/Runtime/AR/Apple/AppleARKit/Source/AppleARKit/Private/AppleARKitTextures.cpp
--- a/Engine/Plugins/Runtime/AR/Apple/AppleARKit/Source/AppleARKit/Private/AppleARKitTextures.cpp
+++ b/Engine/Plugins/Runtime/AR/Apple/AppleARKit/Source/AppleARKit/Private/AppleARKitTextures.cpp
@@ -41,9 +41,16 @@ public:
 		if (CameraImage != nullptr)
 		{
 			SCOPED_AUTORELEASE_POOL;
-			
-			Size.X = CVPixelBufferGetWidth(CameraImage);
-			Size.Y = CVPixelBufferGetHeight(CameraImage);
+
+			CGColorSpaceRef ColorSpaceRef = CGColorSpaceCreateDeviceRGB();
+			CIImage* Image = [[CIImage alloc] initWithCVPixelBuffer: CameraImage];
+
+			// Textures always need to be rotated to the "LandscapeRight" orientation
+			CIImage* RotatedImage = [Image imageByApplyingOrientation: GetRotationFromDeviceOrientation()];
+
+			// Get the size using the rotated image
+			Size.X = RotatedImage.extent.size.width;
+			Size.Y = RotatedImage.extent.size.height;
 
 			// Create the target texture that we'll update into
 			FRHIResourceCreateInfo CreateInfo;
@@ -52,11 +59,8 @@ public:
 			// Get the underlying metal texture so we can render to it
 			id<MTLTexture> UnderlyingMetalTexture = (id<MTLTexture>)DecodedTextureRef->GetNativeResource();
 
-			CIImage* Image = [[CIImage alloc] initWithCVPixelBuffer: CameraImage];
-			CIContext* Context = [CIContext context];
-			CGColorSpaceRef ColorSpaceRef = CGColorSpaceCreateDeviceRGB();
 			// Do the conversion on the GPU
-			[Context render: Image toMTLTexture: UnderlyingMetalTexture commandBuffer: nil bounds: Image.extent colorSpace: ColorSpaceRef];
+			[[CIContext context] render: RotatedImage toMTLTexture: UnderlyingMetalTexture commandBuffer: nil bounds: RotatedImage.extent colorSpace: ColorSpaceRef];
 
 			// Now that the conversion is done, we can get rid of our refs
 			[Image release];
@@ -124,6 +128,36 @@ public:
 #endif
 
 private:
+	/** @return the rotation to use to rotate the texture to the proper direction */
+	int32 GetRotationFromDeviceOrientation()
+	{
+		extern UIInterfaceOrientation GInterfaceOrientation;
+		switch (GInterfaceOrientation)
+		{
+			case UIInterfaceOrientationPortrait:
+			{
+				return kCGImagePropertyOrientationRight;
+			}
+
+			case UIInterfaceOrientationLandscapeLeft:
+			{
+				return kCGImagePropertyOrientationDown;
+			}
+
+			case UIInterfaceOrientationPortraitUpsideDown:
+			{
+				return kCGImagePropertyOrientationLeft;
+			}
+
+			case UIInterfaceOrientationLandscapeRight:
+			{
+				return kCGImagePropertyOrientationUp;
+			}
+		}
+		// Don't know so don't rotate
+		return kCGImagePropertyOrientationUp;
+	}
+
 	/** The size we get from the incoming camera image */
 	FIntPoint Size;
 
