Compiled the ARKit texture code out for Mac, too

diff --git a/Engine/Plugins/Runtime/AR/Apple/AppleARKit/Source/AppleARKit/Private/AppleARKitTextures.cpp b/Engine/Plugins/Runtime/AR/Apple/AppleARKit/Source/AppleARKit/Private/AppleARKitTextures.cpp
--- a/Engine/Plugins/Runtime/AR/Apple/AppleARKit/Source/AppleARKit/Private/AppleARKitTextures.cpp
+++ b/Engine/Plugins/Runtime/AR/Apple/AppleARKit/Source/AppleARKit/Private/AppleARKitTextures.cpp
@@ -37,7 +37,7 @@ public:
 		FSamplerStateInitializerRHI SamplerStateInitializer(SF_Bilinear, AM_Clamp, AM_Clamp, AM_Clamp);
 		SamplerStateRHI = RHICreateSamplerState(SamplerStateInitializer);
 
-#if PLATFORM_MAC || PLATFORM_IOS
+#if PLATFORM_IOS
 		if (CameraImage != nullptr)
 		{
 			SCOPED_AUTORELEASE_POOL;
@@ -128,6 +128,7 @@ public:
 #endif
 
 private:
+#if PLATFORM_IOS
 	/** @return the rotation to use to rotate the texture to the proper direction */
 	int32 GetRotationFromDeviceOrientation()
 	{
@@ -157,7 +158,8 @@ private:
 		// Don't know so don't rotate
 		return kCGImagePropertyOrientationUp;
 	}
-
+#endif
+	
 	/** The size we get from the incoming camera image */
 	FIntPoint Size;
 
