Hooked up the ARKit2 scene understanding API to our AR layer

diff --git a/Engine/Plugins/Runtime/AR/Apple/AppleARKit/Source/AppleARKit/Private/AppleARKitSystem.cpp b/Engine/Plugins/Runtime/AR/Apple/AppleARKit/Source/AppleARKit/Private/AppleARKitSystem.cpp
--- a/Engine/Plugins/Runtime/AR/Apple/AppleARKit/Source/AppleARKit/Private/AppleARKitSystem.cpp
+++ b/Engine/Plugins/Runtime/AR/Apple/AppleARKit/Source/AppleARKit/Private/AppleARKitSystem.cpp
@@ -1518,6 +1518,14 @@ static TSharedPtr<FAppleARKitAnchorData> MakeAnchorData( ARAnchor* Anchor, doubl
 				NewAnchor->Indices.Add(0);
 			}
 		}
+#endif
+#if SUPPORTS_ARKIT_2_0
+		if (FAppleARKitAvailability::SupportsARKit20())
+		{
+			NewAnchor->ObjectClassification = FAppleARKitConversion::ToEARObjectClassification(PlaneAnchor.classification);
+		}
+#else
+		NewAnchor->ObjectClassification = ToEARObjectClassification::Unknown;
 #endif
 	}
 #if SUPPORTS_ARKIT_1_5
@@ -1731,6 +1739,7 @@ void FAppleARKitSystem::SessionDidAddAnchors_Internal( TSharedRef<FAppleARKitAnc
 				// Connect the tracked geo to the MRMesh
 				NewGeo->SetUnderlyingMesh(MRMesh);
 			}
+			NewGeo->SetObjectClassification(AnchorData->ObjectClassification);
 			NewGeometry = NewGeo;
 			break;
 		}
@@ -1867,6 +1876,7 @@ void FAppleARKitSystem::SessionDidUpdateAnchors_Internal( TSharedRef<FAppleARKit
 						Pin->OnTransformUpdated(Pin_LocalToTrackingTransform_PostUpdate);
 					}
 					PlaneGeo->SetOrientation(AnchorData->Orientation);
+					PlaneGeo->SetObjectClassification(AnchorData->ObjectClassification);
 					// Update the occlusion geo if configured
 					if (GetARCompositionComponent()->GetSessionConfig().bGenerateMeshDataFromTrackedGeometry)
 					{
diff --git a/Engine/Plugins/Runtime/AR/Apple/AppleARKit/Source/AppleARKit/Public/AppleARKitConversion.h b/Engine/Plugins/Runtime/AR/Apple/AppleARKit/Source/AppleARKit/Public/AppleARKitConversion.h
--- a/Engine/Plugins/Runtime/AR/Apple/AppleARKit/Source/AppleARKit/Public/AppleARKitConversion.h
+++ b/Engine/Plugins/Runtime/AR/Apple/AppleARKit/Source/AppleARKit/Public/AppleARKitConversion.h
@@ -184,6 +184,40 @@ struct APPLEARKIT_API FAppleARKitConversion
 		return RetVal;
 	}
 
+#if SUPPORTS_ARKIT_2_0
+	static FORCEINLINE EARObjectClassification ToEARObjectClassification(ARPlaneClassification Classification)
+	{
+		switch(Classification)
+		{
+			case ARPlaneClassificationWall:
+			{
+				return EARObjectClassification::Wall;
+			}
+
+			case ARPlaneClassificationFloor:
+			{
+				return EARObjectClassification::Floor;
+			}
+
+			case ARPlaneClassificationCeiling:
+			{
+				return EARObjectClassification::Ceiling;
+			}
+
+			case ARPlaneClassificationTable:
+			{
+				return EARObjectClassification::Table;
+			}
+
+			case ARPlaneClassificationSeat:
+			{
+				return EARObjectClassification::Seat;
+			}
+		}
+		return EARObjectClassification::Unknown;
+	}
+#endif
+
 #pragma clang diagnostic push
 #pragma clang diagnostic ignored "-Wpartial-availability"
 
@@ -329,6 +363,11 @@ struct FAppleARKitAnchorData
 		Timecode = Other.Timecode;
 		FrameRate = Other.FrameRate;
 
+		Vertices = Other.Vertices;
+		Indices = Other.Indices;
+		Orientation = Other.Orientation;
+		ObjectClassification = Other.ObjectClassification;
+
 		bIsTracked = Other.bIsTracked;
 	}
 
@@ -337,6 +376,8 @@ struct FAppleARKitAnchorData
 		BoundaryVerts.Empty();
 		BlendShapes.Empty();
 		FaceVerts.Empty();
+		Vertices.Empty();
+		Indices.Empty();
 		ProbeTexture = nullptr;
 	}
 
@@ -346,6 +387,7 @@ struct FAppleARKitAnchorData
 	FVector Center;
 	FVector Extent;
 	EARPlaneOrientation Orientation;
+	EARObjectClassification ObjectClassification;
 	/** Set by the session config to detemine whether to generate geometry or not */
 	static bool bGenerateGeometry;
 	TArray<FVector> BoundaryVerts;
