Fixed missing image and object scanning session type checks

diff --git a/Engine/Plugins/Runtime/AR/Apple/AppleARKit/Source/AppleARKit/Private/AppleARKitSystem.cpp b/Engine/Plugins/Runtime/AR/Apple/AppleARKit/Source/AppleARKit/Private/AppleARKitSystem.cpp
--- a/Engine/Plugins/Runtime/AR/Apple/AppleARKit/Source/AppleARKit/Private/AppleARKitSystem.cpp
+++ b/Engine/Plugins/Runtime/AR/Apple/AppleARKit/Source/AppleARKit/Private/AppleARKitSystem.cpp
@@ -809,6 +809,16 @@ bool FAppleARKitSystem::OnIsTrackingTypeSupported(EARSessionType SessionType) co
 			}
 			return false;
 		}
+#if SUPPORTS_ARKIT_2_0
+		case EARSessionType::Image:
+		{
+			return ARImageTrackingConfiguration.isSupported == TRUE;
+		}
+		case EARSessionType::ObjectScanning:
+		{
+			return ARObjectScanningConfiguration.isSupported == TRUE;
+		}
+#endif
 	}
 #endif
 	return false;
