Fixed missing retain for the ReferenceObject that is passed into the completion handler by ARKit

diff --git a/Engine/Plugins/Runtime/AR/Apple/AppleARKit/Source/AppleARKit/Private/AppleARKitSystem.cpp b/Engine/Plugins/Runtime/AR/Apple/AppleARKit/Source/AppleARKit/Private/AppleARKitSystem.cpp
--- a/Engine/Plugins/Runtime/AR/Apple/AppleARKit/Source/AppleARKit/Private/AppleARKitSystem.cpp
+++ b/Engine/Plugins/Runtime/AR/Apple/AppleARKit/Source/AppleARKit/Private/AppleARKitSystem.cpp
@@ -1004,6 +1004,7 @@ public:
 		 completionHandler: ^(ARReferenceObject* refObject, NSError* error)
 		{
 			ReferenceObject = refObject;
+			CFRetain(ReferenceObject);
 			bool bWasSuccessful = error == nullptr;
 			bHadError = error != nullptr;
 			FString ErrorString;
