Remove all tracked geometry immediately when the user restarts a session and has the reset tracked geometry flag set to true instead of waiting for ARKit to send the notifications of removal
Changed the removal notification to handle removing items that we removed manually

diff --git a/Engine/Plugins/Runtime/AR/Apple/AppleARKit/Source/AppleARKit/Private/AppleARKitSystem.cpp b/Engine/Plugins/Runtime/AR/Apple/AppleARKit/Source/AppleARKit/Private/AppleARKitSystem.cpp
--- a/Engine/Plugins/Runtime/AR/Apple/AppleARKit/Source/AppleARKit/Private/AppleARKitSystem.cpp
+++ b/Engine/Plugins/Runtime/AR/Apple/AppleARKit/Source/AppleARKit/Private/AppleARKitSystem.cpp
@@ -1269,6 +1269,14 @@ static EScreenOrientation::Type GetAppOrientation()
 #endif
 }
 
+void FAppleARKitSystem::ClearTrackedGeometries()
+{
+	for (auto GeoIt=TrackedGeometries.CreateIterator(); GeoIt; ++GeoIt)
+	{
+		SessionDidRemoveAnchors_Internal(GeoIt.Key());
+	}
+}
+
 PRAGMA_DISABLE_OPTIMIZATION
 bool FAppleARKitSystem::Run(UARSessionConfig* SessionConfig)
 {
@@ -1338,6 +1346,8 @@ bool FAppleARKitSystem::Run(UARSessionConfig* SessionConfig)
 			if (SessionConfig->ShouldResetTrackedObjects())
 			{
 				options |= ARSessionRunOptionRemoveExistingAnchors;
+				// The user requested us to remove existing anchors so remove ours now
+				ClearTrackedGeometries();
 			}
 
 			[Session pause];
@@ -1978,23 +1988,28 @@ void FAppleARKitSystem::SessionDidRemoveAnchors_Internal( FGuid AnchorGuid )
 
 	// Notify pin that it is being orphaned
 	{
-		UARTrackedGeometry* TrackedGeometryBeingRemoved = TrackedGeometries.FindChecked(AnchorGuid);
-		TrackedGeometryBeingRemoved->UpdateTrackingState(EARTrackingState::StoppedTracking);
-		// Remove the occlusion mesh if present
-		UMRMeshComponent* MRMesh = TrackedGeometryBeingRemoved->GetUnderlyingMesh();
-		if (MRMesh != nullptr)
+		UARTrackedGeometry** FoundGeo = TrackedGeometries.Find(AnchorGuid);
+		// This no longer performs a FindChecked() because the act of discard on restart can cause this to be missing
+		if (FoundGeo != nullptr)
 		{
-			MRMesh->UnregisterComponent();
-			TrackedGeometryBeingRemoved->SetUnderlyingMesh(nullptr);
-		}
+			UARTrackedGeometry* TrackedGeometryBeingRemoved = *FoundGeo;
+			TrackedGeometryBeingRemoved->UpdateTrackingState(EARTrackingState::StoppedTracking);
+			// Remove the occlusion mesh if present
+			UMRMeshComponent* MRMesh = TrackedGeometryBeingRemoved->GetUnderlyingMesh();
+			if (MRMesh != nullptr)
+			{
+				MRMesh->UnregisterComponent();
+				TrackedGeometryBeingRemoved->SetUnderlyingMesh(nullptr);
+			}
 
-		TArray<UARPin*> ARPinsBeingOrphaned = ARKitUtil::PinsFromGeometry(TrackedGeometryBeingRemoved, Pins);
-		for(UARPin* PinBeingOrphaned : ARPinsBeingOrphaned)
-		{
-			PinBeingOrphaned->OnTrackingStateChanged(EARTrackingState::StoppedTracking);
+			TArray<UARPin*> ARPinsBeingOrphaned = ARKitUtil::PinsFromGeometry(TrackedGeometryBeingRemoved, Pins);
+			for(UARPin* PinBeingOrphaned : ARPinsBeingOrphaned)
+			{
+				PinBeingOrphaned->OnTrackingStateChanged(EARTrackingState::StoppedTracking);
+			}
+			// Trigger the delegate so anyone listening can take action
+			TriggerOnTrackableRemovedDelegates(TrackedGeometryBeingRemoved);
 		}
-		// Trigger the delegate so anyone listening can take action
-		TriggerOnTrackableRemovedDelegates(TrackedGeometryBeingRemoved);
 	}
 	
 	TrackedGeometries.Remove(AnchorGuid);
diff --git a/Engine/Plugins/Runtime/AR/Apple/AppleARKit/Source/AppleARKit/Private/AppleARKitSystem.h b/Engine/Plugins/Runtime/AR/Apple/AppleARKit/Source/AppleARKit/Private/AppleARKitSystem.h
--- a/Engine/Plugins/Runtime/AR/Apple/AppleARKit/Source/AppleARKit/Private/AppleARKitSystem.h
+++ b/Engine/Plugins/Runtime/AR/Apple/AppleARKit/Source/AppleARKit/Private/AppleARKitSystem.h
@@ -116,7 +116,8 @@ private:
 	void SessionDidRemoveAnchors_Internal( FGuid AnchorGuid );
 #endif
 	void SessionDidUpdateFrame_Internal( TSharedRef< FAppleARKitFrame, ESPMode::ThreadSafe > Frame );
-
+	/** Removes all tracked geometries, marking them as not tracked and sending the delegate event */
+	void ClearTrackedGeometries();
 	
 public:
 	/**
