Changed how tracked geometries are removed when restarting a AR session on ARKit

diff --git a/Engine/Plugins/Runtime/AR/Apple/AppleARKit/Source/AppleARKit/Private/AppleARKitSystem.cpp b/Engine/Plugins/Runtime/AR/Apple/AppleARKit/Source/AppleARKit/Private/AppleARKitSystem.cpp
--- a/Engine/Plugins/Runtime/AR/Apple/AppleARKit/Source/AppleARKit/Private/AppleARKitSystem.cpp
+++ b/Engine/Plugins/Runtime/AR/Apple/AppleARKit/Source/AppleARKit/Private/AppleARKitSystem.cpp
@@ -1272,9 +1272,11 @@ static EScreenOrientation::Type GetAppOrientation()
 void FAppleARKitSystem::ClearTrackedGeometries()
 {
 #if SUPPORTS_ARKIT_1_0
-	for (auto GeoIt=TrackedGeometries.CreateIterator(); GeoIt; ++GeoIt)
+	TArray<FGuid> Keys;
+	TrackedGeometries.GetKeys(Keys);
+	for (const FGuid& Key : Keys)
 	{
-		SessionDidRemoveAnchors_Internal(GeoIt.Key());
+		SessionDidRemoveAnchors_Internal(Key);
 	}
 #endif
 }
