Fixed mismatched sRGB flags in ARKit environment capture probes

diff --git a/Engine/Plugins/Runtime/AR/Apple/AppleARKit/Source/AppleARKit/Private/AppleARKitTextures.cpp b/Engine/Plugins/Runtime/AR/Apple/AppleARKit/Source/AppleARKit/Private/AppleARKitTextures.cpp
--- a/Engine/Plugins/Runtime/AR/Apple/AppleARKit/Source/AppleARKit/Private/AppleARKitTextures.cpp
+++ b/Engine/Plugins/Runtime/AR/Apple/AppleARKit/Source/AppleARKit/Private/AppleARKitTextures.cpp
@@ -400,7 +400,7 @@ public:
 		: Owner(InOwner)
 	{
 		bGreyScaleFormat = false;
-		bSRGB = true;
+		bSRGB = InOwner->SRGB;
 	}
 	
 	virtual ~FARMetalResource()
