Removed the pause of the existing session on ARKit when restarting sessions as it was not needed and introduced delay

diff --git a/Engine/Plugins/Runtime/AR/Apple/AppleARKit/Source/AppleARKit/Private/AppleARKitSystem.cpp b/Engine/Plugins/Runtime/AR/Apple/AppleARKit/Source/AppleARKit/Private/AppleARKitSystem.cpp
--- a/Engine/Plugins/Runtime/AR/Apple/AppleARKit/Source/AppleARKit/Private/AppleARKitSystem.cpp
+++ b/Engine/Plugins/Runtime/AR/Apple/AppleARKit/Source/AppleARKit/Private/AppleARKitSystem.cpp
@@ -1353,8 +1353,6 @@ bool FAppleARKitSystem::Run(UARSessionConfig* SessionConfig)
 				// The user requested us to remove existing anchors so remove ours now
 				ClearTrackedGeometries();
 			}
-
-			[Session pause];
 		}
 
 		// Create MetalTextureCache
