Updated the AR rotation uses to match the change to using UI orientation instead of device orientation

diff --git a/Engine/Plugins/Runtime/AR/Apple/AppleARKit/Source/AppleARKit/Private/AppleARKitSystem.cpp b/Engine/Plugins/Runtime/AR/Apple/AppleARKit/Source/AppleARKit/Private/AppleARKitSystem.cpp
--- a/Engine/Plugins/Runtime/AR/Apple/AppleARKit/Source/AppleARKit/Private/AppleARKitSystem.cpp
+++ b/Engine/Plugins/Runtime/AR/Apple/AppleARKit/Source/AppleARKit/Private/AppleARKitSystem.cpp
@@ -99,7 +99,7 @@ private:
 		const bool bShouldOverrideFOV = ARKitSystem.GetARCompositionComponent()->GetSessionConfig().ShouldRenderCameraOverlay();
 		if (bShouldOverrideFOV && ARKitSystem.GameThreadFrame.IsValid())
 		{
-			if (ARKitSystem.DeviceOrientation == EScreenOrientation::Portrait || ARKitSystem.DeviceOrientation == EScreenOrientation::PortraitUpsideDown)
+			if (ARKitSystem.DeviceOrientation == EDeviceScreenOrientation::Portrait || ARKitSystem.DeviceOrientation == EDeviceScreenOrientation::PortraitUpsideDown)
 			{
 				// Portrait
 				InOutFOV = ARKitSystem.GameThreadFrame->Camera.GetVerticalFieldOfViewForScreen(EAppleARKitBackgroundFitMode::Fill);
@@ -208,7 +208,7 @@ private:
 
 FAppleARKitSystem::FAppleARKitSystem()
 : FXRTrackingSystemBase(this)
-, DeviceOrientation(EScreenOrientation::Unknown)
+, DeviceOrientation(EDeviceScreenOrientation::Unknown)
 , DerivedTrackingToUnrealRotation(FRotator::ZeroRotator)
 , LightEstimate(nullptr)
 , CameraImage(nullptr)
@@ -321,19 +321,19 @@ void FAppleARKitSystem::CalcTrackingToWorldRotation()
 	{
 		switch (DeviceOrientation)
 		{
-			case EScreenOrientation::Portrait:
+			case EDeviceScreenOrientation::Portrait:
 				DerivedTrackingToUnrealRotation = FRotator(0.0f, 0.0f, -90.0f);
 				break;
 				
-			case EScreenOrientation::PortraitUpsideDown:
+			case EDeviceScreenOrientation::PortraitUpsideDown:
 				DerivedTrackingToUnrealRotation = FRotator(0.0f, 0.0f, 90.0f);
 				break;
 				
 			default:
-			case EScreenOrientation::LandscapeLeft:
+			case EDeviceScreenOrientation::LandscapeRight:
 				break;
 				
-			case EScreenOrientation::LandscapeRight:
+			case EDeviceScreenOrientation::LandscapeLeft:
 				DerivedTrackingToUnrealRotation = FRotator(0.0f, 0.0f, 180.0f);
 				break;
 		}
@@ -343,20 +343,20 @@ void FAppleARKitSystem::CalcTrackingToWorldRotation()
 	{
 		switch (DeviceOrientation)
 		{
-			case EScreenOrientation::Portrait:
+			case EDeviceScreenOrientation::Portrait:
 				DerivedTrackingToUnrealRotation = FRotator(0.0f, 0.0f, 90.0f);
 				break;
 				
-			case EScreenOrientation::PortraitUpsideDown:
+			case EDeviceScreenOrientation::PortraitUpsideDown:
 				DerivedTrackingToUnrealRotation = FRotator(0.0f, 0.0f, -90.0f);
 				break;
 				
 			default:
-			case EScreenOrientation::LandscapeLeft:
+			case EDeviceScreenOrientation::LandscapeLeft:
 				DerivedTrackingToUnrealRotation = FRotator(0.0f, 0.0f, -180.0f);
 				break;
 				
-			case EScreenOrientation::LandscapeRight:
+			case EDeviceScreenOrientation::LandscapeRight:
 				break;
 		}
 	}
@@ -610,19 +610,19 @@ TArray<FARTraceResult> FAppleARKitSystem::OnLineTraceTrackedObjects( const FVect
 				FVector2D NormalizedImagePosition = FAppleARKitCamera( HitTestFrame.camera ).GetImageCoordinateForScreenPosition( ScreenCoord, EAppleARKitBackgroundFitMode::Fill );
 				switch (DeviceOrientation)
 				{
-					case EScreenOrientation::Portrait:
+					case EDeviceScreenOrientation::Portrait:
 						NormalizedImagePosition = FVector2D( NormalizedImagePosition.Y, 1.0f - NormalizedImagePosition.X );
 						break;
 						
-					case EScreenOrientation::PortraitUpsideDown:
+					case EDeviceScreenOrientation::PortraitUpsideDown:
 						NormalizedImagePosition = FVector2D( 1.0f - NormalizedImagePosition.Y, NormalizedImagePosition.X );
 						break;
 						
 					default:
-					case EScreenOrientation::LandscapeLeft:
+					case EDeviceScreenOrientation::LandscapeRight:
 						break;
 						
-					case EScreenOrientation::LandscapeRight:
+					case EDeviceScreenOrientation::LandscapeLeft:
 						NormalizedImagePosition = FVector2D(1.0f, 1.0f) - NormalizedImagePosition;
 						break;
 				};
@@ -1150,125 +1150,22 @@ bool FAppleARKitSystem::HitTestAtScreenPosition(const FVector2D ScreenPosition,
 	return false;
 }
 
-static TOptional<EScreenOrientation::Type> PickAllowedDeviceOrientation( EScreenOrientation::Type InOrientation )
+void FAppleARKitSystem::SetDeviceOrientation(EDeviceScreenOrientation InOrientation)
 {
-#if SUPPORTS_ARKIT_1_0
-	const UIOSRuntimeSettings* IOSSettings = GetDefault<UIOSRuntimeSettings>();
-	
-	const bool bOrientationSupported[] =
-	{
-		true, // Unknown
-		IOSSettings->bSupportsPortraitOrientation != 0, // Portait
-		IOSSettings->bSupportsUpsideDownOrientation != 0, // PortraitUpsideDown
-		IOSSettings->bSupportsLandscapeRightOrientation != 0, // LandscapeLeft; These are flipped vs the enum name?
-		IOSSettings->bSupportsLandscapeLeftOrientation != 0, // LandscapeRight; These are flipped vs the enum name?
-		false, // FaceUp
-		false // FaceDown
-	};
-	
-	if (bOrientationSupported[static_cast<int32>(InOrientation)])
-	{
-		return InOrientation;
-	}
-	else
+	ensureAlwaysMsgf(InOrientation != EDeviceScreenOrientation::Unknown, TEXT("statusBarOrientation should only ever return valid orientations"));
+	if (InOrientation == EDeviceScreenOrientation::Unknown)
 	{
-		return TOptional<EScreenOrientation::Type>();
+		// This is the default for AR apps
+		InOrientation = EDeviceScreenOrientation::LandscapeLeft;
 	}
-#else
-	return TOptional<EScreenOrientation::Type>();
-#endif
-}
-
-void FAppleARKitSystem::SetDeviceOrientation( EScreenOrientation::Type InOrientation )
-{
-	TOptional<EScreenOrientation::Type> NewOrientation = PickAllowedDeviceOrientation(InOrientation);
-
-	if (!NewOrientation.IsSet() && DeviceOrientation == EScreenOrientation::Unknown)
-	{
-		// We do not currently have a valid orientation, nor did the device provide one.
-		// So pick ANY ALLOWED default.
-		// This only realy happens if the device is face down on something or
-		// in another "useless" state for AR.
-
-		// Note: the order in which this selection is done is important and must match that 
-		// established in UEDeployIOS.cs and written into UISupportedInterfaceOrientations.
-		// IOSView preferredInterfaceOrientationForPresentation presumably also should match.
-		//
-		// However it would very likely be better to hook statusBarOrientation instead of deviceOrientation to update
-		// our orientation, in which case we would only need to handle unknown.
-		
-		if (!NewOrientation.IsSet())
-		{
-			NewOrientation = PickAllowedDeviceOrientation(EScreenOrientation::Portrait);
-		}
-
-		if (!NewOrientation.IsSet())
-		{
-			NewOrientation = PickAllowedDeviceOrientation(EScreenOrientation::PortraitUpsideDown);
-		}
 
-#if SUPPORTS_ARKIT_1_0
-		const UIOSRuntimeSettings* IOSSettings = GetDefault<UIOSRuntimeSettings>();
-		const bool bPreferLandscapeLeftHomeButton = IOSSettings->PreferredLandscapeOrientation == EIOSLandscapeOrientation::LandscapeLeft;
-#else
-		const bool bPreferLandscapeLeftHomeButton = true;
-#endif
-		if (bPreferLandscapeLeftHomeButton)
-		{
-			if (!NewOrientation.IsSet())
-			{
-				NewOrientation = PickAllowedDeviceOrientation(EScreenOrientation::LandscapeRight);
-			}
-			if (!NewOrientation.IsSet())
-			{
-				NewOrientation = PickAllowedDeviceOrientation(EScreenOrientation::LandscapeLeft);
-			}
-		}
-		else
-		{
-			if (!NewOrientation.IsSet())
-			{
-				NewOrientation = PickAllowedDeviceOrientation(EScreenOrientation::LandscapeLeft);
-			}
-			if (!NewOrientation.IsSet())
-			{
-				NewOrientation = PickAllowedDeviceOrientation(EScreenOrientation::LandscapeRight);
-			}
-		}
-		
-		check(NewOrientation.IsSet());
-	}
-	
-	if (NewOrientation.IsSet() && DeviceOrientation != NewOrientation.GetValue())
+	if (DeviceOrientation != InOrientation)
 	{
-		DeviceOrientation = NewOrientation.GetValue();
+		DeviceOrientation = InOrientation;
 		CalcTrackingToWorldRotation();
 	}
 }
 
-static EScreenOrientation::Type GetAppOrientation()
-{
-#if PLATFORM_IOS && !PLATFORM_TVOS
-	// We want the orientation that the app is running with, not necessarily the orientation of the device right now.
-	UIInterfaceOrientation Orientation = [[UIApplication sharedApplication] statusBarOrientation];
-#if __IPHONE_OS_VERSION_MIN_REQUIRED < __IPHONE_8_0
-	Orientation = [[IOSAppDelegate GetDelegate].IOSController interfaceOrientation];
-#endif
-	EScreenOrientation::Type ScreenOrientation = EScreenOrientation::Unknown;
-	switch (Orientation)
-	{
-	case UIInterfaceOrientationUnknown:				return EScreenOrientation::Unknown;
-	case UIInterfaceOrientationPortrait:			return EScreenOrientation::Portrait;
-	case UIInterfaceOrientationPortraitUpsideDown:	return EScreenOrientation::PortraitUpsideDown;
-	case UIInterfaceOrientationLandscapeLeft:		return EScreenOrientation::LandscapeRight;
-	case UIInterfaceOrientationLandscapeRight:		return EScreenOrientation::LandscapeLeft;
-	default:										check(false); return EScreenOrientation::Unknown;
-	}
-#else
-	return static_cast<EScreenOrientation::Type>(FPlatformMisc::GetDeviceOrientation());
-#endif
-}
-
 void FAppleARKitSystem::ClearTrackedGeometries()
 {
 #if SUPPORTS_ARKIT_1_0
@@ -1294,9 +1191,9 @@ bool FAppleARKitSystem::Run(UARSessionConfig* SessionConfig)
 	}
 
 	// Make sure this is set at session start, because there are timing issues with using only the delegate approach
-	if (DeviceOrientation == EScreenOrientation::Unknown)
+	if (DeviceOrientation == EDeviceScreenOrientation::Unknown)
 	{
-		const EScreenOrientation::Type ScreenOrientation = GetAppOrientation();
+		EDeviceScreenOrientation ScreenOrientation = FPlatformMisc::GetDeviceOrientation();
 		SetDeviceOrientation( ScreenOrientation );
 	}
 
@@ -1461,7 +1358,7 @@ bool FAppleARKitSystem::Pause()
 
 void FAppleARKitSystem::OrientationChanged(const int32 NewOrientationRaw)
 {
-	const EScreenOrientation::Type NewOrientation = static_cast<EScreenOrientation::Type>(NewOrientationRaw);
+	const EDeviceScreenOrientation NewOrientation = static_cast<EDeviceScreenOrientation>(NewOrientationRaw);
 	SetDeviceOrientation(NewOrientation);
 }
 						
diff --git a/Engine/Plugins/Runtime/AR/Apple/AppleARKit/Source/AppleARKit/Private/AppleARKitSystem.h b/Engine/Plugins/Runtime/AR/Apple/AppleARKit/Source/AppleARKit/Private/AppleARKitSystem.h
--- a/Engine/Plugins/Runtime/AR/Apple/AppleARKit/Source/AppleARKit/Private/AppleARKitSystem.h
+++ b/Engine/Plugins/Runtime/AR/Apple/AppleARKit/Source/AppleARKit/Private/AppleARKitSystem.h
@@ -139,7 +139,7 @@ private:
 	
 	bool bIsRunning = false;
 	
-	void SetDeviceOrientation( EScreenOrientation::Type InOrientation );
+	void SetDeviceOrientation(EDeviceScreenOrientation InOrientation);
 
 	/** Creates or clears the face ar support object if face ar has been requested */
 	void CheckForFaceARSupport(UARSessionConfig* InSessionConfig);
@@ -147,8 +147,8 @@ private:
 	/** Updates the ARKit perf counters */
 	void UpdateARKitPerfStats();
 	
-	/** The orientation of the device; see EScreenOrientation */
-	EScreenOrientation::Type DeviceOrientation;
+	/** The orientation of the device; see EDeviceScreenOrientation */
+	EDeviceScreenOrientation DeviceOrientation;
 	
 	/** A rotation from ARKit TrackingSpace to Unreal Space. It is re-derived based on other parameters; users should not set it directly. */
 	FRotator DerivedTrackingToUnrealRotation;
diff --git a/Engine/Plugins/Runtime/AR/Apple/AppleARKit/Source/AppleARKit/Private/AppleARKitVideoOverlay.cpp b/Engine/Plugins/Runtime/AR/Apple/AppleARKit/Source/AppleARKit/Private/AppleARKitVideoOverlay.cpp
--- a/Engine/Plugins/Runtime/AR/Apple/AppleARKit/Source/AppleARKit/Private/AppleARKitVideoOverlay.cpp
+++ b/Engine/Plugins/Runtime/AR/Apple/AppleARKit/Source/AppleARKit/Private/AppleARKitVideoOverlay.cpp
@@ -375,7 +375,7 @@ private:
 IMPLEMENT_MATERIAL_SHADER_TYPE(template<>, FARKitCameraOverlayPS<true>, TEXT("/Engine/Private/PostProcessMaterialShaders.usf"), TEXT("MainPS_ES2"), SF_Pixel);
 IMPLEMENT_MATERIAL_SHADER_TYPE(template<>, FARKitCameraOverlayPS<false>, TEXT("/Engine/Private/PostProcessMaterialShaders.usf"), TEXT("MainPS_VideoOverlay"), SF_Pixel);
 
-void FAppleARKitVideoOverlay::RenderVideoOverlay_RenderThread(FRHICommandListImmediate& RHICmdList, const FSceneView& InView, const EScreenOrientation::Type DeviceOrientation)
+void FAppleARKitVideoOverlay::RenderVideoOverlay_RenderThread(FRHICommandListImmediate& RHICmdList, const FSceneView& InView, const EDeviceScreenOrientation DeviceOrientation)
 {
 #if SUPPORTS_ARKIT_1_0
 
@@ -447,19 +447,19 @@ void FAppleARKitVideoOverlay::RenderVideoOverlay_RenderThread(FRHICommandListImm
 		FVertexBufferRHIParamRef VertexBufferRHI = nullptr;
 		switch (DeviceOrientation)
 		{
-			case EScreenOrientation::Type::LandscapeLeft:
+			case EDeviceScreenOrientation::LandscapeRight:
 				VertexBufferRHI = OverlayVertexBufferRHI[0];
 				break;
 
-			case EScreenOrientation::Type::LandscapeRight:
+			case EDeviceScreenOrientation::LandscapeLeft:
 				VertexBufferRHI = OverlayVertexBufferRHI[1];
 				break;
 
-			case EScreenOrientation::Type::Portrait:
+			case EDeviceScreenOrientation::Portrait:
 				VertexBufferRHI = OverlayVertexBufferRHI[2];
 				break;
 
-			case EScreenOrientation::Type::PortraitUpsideDown:
+			case EDeviceScreenOrientation::PortraitUpsideDown:
 				VertexBufferRHI = OverlayVertexBufferRHI[3];
 				break;
 
@@ -486,7 +486,7 @@ void FAppleARKitVideoOverlay::RenderVideoOverlay_RenderThread(FRHICommandListImm
 }
 
 
-bool FAppleARKitVideoOverlay::GetPassthroughCameraUVs_RenderThread(TArray<FVector2D>& OutUVs, const EScreenOrientation::Type DeviceOrientation)
+bool FAppleARKitVideoOverlay::GetPassthroughCameraUVs_RenderThread(TArray<FVector2D>& OutUVs, const EDeviceScreenOrientation DeviceOrientation)
 {
 #if SUPPORTS_ARKIT_1_0
 	if (VideoTextureY != nullptr) // make sure UVOffset has been initialized
@@ -495,28 +495,28 @@ bool FAppleARKitVideoOverlay::GetPassthroughCameraUVs_RenderThread(TArray<FVecto
 
 		switch (DeviceOrientation)
 		{
-		case EScreenOrientation::Type::LandscapeLeft:
+		case EDeviceScreenOrientation::LandscapeRight:
 			OutUVs[1] = FVector2D(UVOffset.X, 1.0f - UVOffset.Y);
 			OutUVs[0] = FVector2D(UVOffset.X, UVOffset.Y);
 			OutUVs[3] = FVector2D(1.0f - UVOffset.X, 1.0f - UVOffset.Y);
 			OutUVs[2] = FVector2D(1.0f - UVOffset.X, UVOffset.Y);
 			return true;
 
-		case EScreenOrientation::Type::LandscapeRight:
+		case EDeviceScreenOrientation::LandscapeLeft:
 			OutUVs[1] = FVector2D(1.0f - UVOffset.X, UVOffset.Y);
 			OutUVs[0] = FVector2D(1.0f - UVOffset.X, 1.0f - UVOffset.Y);
 			OutUVs[3] = FVector2D(UVOffset.X, UVOffset.Y);
 			OutUVs[2] = FVector2D(UVOffset.X, 1.0f - UVOffset.Y);
 			return true;
 
-		case EScreenOrientation::Type::Portrait:
+		case EDeviceScreenOrientation::Portrait:
 			OutUVs[1] = FVector2D(1.0f - UVOffset.X, 1.0f - UVOffset.Y);
 			OutUVs[0] = FVector2D(UVOffset.X, 1.0f - UVOffset.Y);
 			OutUVs[3] = FVector2D(1.0f - UVOffset.X, UVOffset.Y);
 			OutUVs[2] = FVector2D(UVOffset.X, UVOffset.Y);
 			return true;
 
-		case EScreenOrientation::Type::PortraitUpsideDown:
+		case EDeviceScreenOrientation::PortraitUpsideDown:
 			OutUVs[1] = FVector2D(UVOffset.X, UVOffset.Y);
 			OutUVs[0] = FVector2D(1.0f - UVOffset.X, UVOffset.Y);
 			OutUVs[3] = FVector2D(UVOffset.X, 1.0f - UVOffset.Y);
diff --git a/Engine/Plugins/Runtime/AR/Apple/AppleARKit/Source/AppleARKit/Private/AppleARKitVideoOverlay.h b/Engine/Plugins/Runtime/AR/Apple/AppleARKit/Source/AppleARKit/Private/AppleARKitVideoOverlay.h
--- a/Engine/Plugins/Runtime/AR/Apple/AppleARKit/Source/AppleARKit/Private/AppleARKitVideoOverlay.h
+++ b/Engine/Plugins/Runtime/AR/Apple/AppleARKit/Source/AppleARKit/Private/AppleARKitVideoOverlay.h
@@ -35,8 +35,8 @@ public:
 	FAppleARKitVideoOverlay();
 
 	void UpdateVideoTexture_RenderThread(FRHICommandListImmediate& RHICmdList, FAppleARKitFrame& Frame, const FSceneViewFamily& InViewFamily);
-	void RenderVideoOverlay_RenderThread(FRHICommandListImmediate& RHICmdList, const FSceneView& InView, const EScreenOrientation::Type DeviceOrientation);
-	bool GetPassthroughCameraUVs_RenderThread(TArray<FVector2D>& OutUVs, const EScreenOrientation::Type DeviceOrientation);
+	void RenderVideoOverlay_RenderThread(FRHICommandListImmediate& RHICmdList, const FSceneView& InView, const EDeviceScreenOrientation DeviceOrientation);
+	bool GetPassthroughCameraUVs_RenderThread(TArray<FVector2D>& OutUVs, const EDeviceScreenOrientation DeviceOrientation);
 
 private:
 	FTextureRHIRef VideoTextureY;
