Removed an ensure from ARKit updates that was no longer valid post session restart

diff --git a/Engine/Plugins/Runtime/AR/Apple/AppleARKit/Source/AppleARKit/Private/AppleARKitSystem.cpp b/Engine/Plugins/Runtime/AR/Apple/AppleARKit/Source/AppleARKit/Private/AppleARKitSystem.cpp
--- a/Engine/Plugins/Runtime/AR/Apple/AppleARKit/Source/AppleARKit/Private/AppleARKitSystem.cpp
+++ b/Engine/Plugins/Runtime/AR/Apple/AppleARKit/Source/AppleARKit/Private/AppleARKitSystem.cpp
@@ -1759,7 +1759,7 @@ void FAppleARKitSystem::SessionDidUpdateAnchors_Internal( TSharedRef<FAppleARKit
 	}
 
 	UARTrackedGeometry** GeometrySearchResult = TrackedGeometries.Find(AnchorData->AnchorGUID);
-	if (ensure(GeometrySearchResult != nullptr))
+	if (GeometrySearchResult != nullptr)
 	{
 		UARTrackedGeometry* FoundGeometry = *GeometrySearchResult;
 		TArray<UARPin*> PinsToUpdate = ARKitUtil::PinsFromGeometry(FoundGeometry, Pins);
