Reworked ARKit pass through camera rendering to use the UARTextureCameraImage that is updated per frame on the render thread

diff --git a/Engine/Plugins/Runtime/AR/Apple/AppleARKit/Source/AppleARKit/Private/AppleARKitSystem.cpp b/Engine/Plugins/Runtime/AR/Apple/AppleARKit/Source/AppleARKit/Private/AppleARKitSystem.cpp
--- a/Engine/Plugins/Runtime/AR/Apple/AppleARKit/Source/AppleARKit/Private/AppleARKitSystem.cpp
+++ b/Engine/Plugins/Runtime/AR/Apple/AppleARKit/Source/AppleARKit/Private/AppleARKitSystem.cpp
@@ -89,6 +89,11 @@ public:
 	{
 		ThreadPriority.Set(NewPriority);
 	}
+
+	void SetOverlayTexture(UARTextureCameraImage* InCameraImage)
+	{
+		VideoOverlay.SetOverlayTexture(InCameraImage);
+	}
 	
 private:
 	//~ FDefaultXRCamera
@@ -140,25 +145,12 @@ private:
 	
 	virtual void PreRenderViewFamily_RenderThread(FRHICommandListImmediate& RHICmdList, FSceneViewFamily& InViewFamily) override
 	{
-		// Grab the latest frame from ARKit
-		{
-			FScopeLock ScopeLock(&ARKitSystem.FrameLock);
-			ARKitSystem.RenderThreadFrame = ARKitSystem.LastReceivedFrame;
-		}
-		
-		// @todo arkit: Camera late update? 
-		
-		if (ARKitSystem.RenderThreadFrame.IsValid())
-		{
-			VideoOverlay.UpdateVideoTexture_RenderThread(RHICmdList, *ARKitSystem.RenderThreadFrame, InViewFamily);
-		}
-		
 		FDefaultXRCamera::PreRenderViewFamily_RenderThread(RHICmdList, InViewFamily);
 	}
 	
 	virtual void PostRenderBasePass_RenderThread(FRHICommandListImmediate& RHICmdList, FSceneView& InView) override
 	{
-		VideoOverlay.RenderVideoOverlay_RenderThread(RHICmdList, InView, ARKitSystem.DeviceOrientation);
+		VideoOverlay.RenderVideoOverlay_RenderThread(RHICmdList, InView);
 	}
 	
 	virtual bool GetPassthroughCameraUVs_RenderThread(TArray<FVector2D>& OutUVs) override
@@ -183,13 +175,8 @@ private:
 		{
 			return bRenderOverlay;
 		}
-		else
-		{
-			return false;
-		}
-#else
-		return false;
 #endif
+		return false;
 	}
 	//~ FDefaultXRCamera
 	
@@ -375,22 +362,12 @@ void FAppleARKitSystem::UpdateFrame()
 #if SUPPORTS_ARKIT_1_0
 			if (GameThreadFrame->CameraImage != nullptr)
 			{
-				// Only create a new camera image texture if it's set and we don't already have one
-				if (CameraImage == nullptr)
-				{
-					CameraImage = NewObject<UAppleARKitTextureCameraImage>();
-				}
 				// Reuse the UObjects because otherwise the time between GCs causes ARKit to be starved of resources
                 CameraImage->Init(FPlatformTime::Seconds(), GameThreadFrame->CameraImage);
 			}
 
 			if (GameThreadFrame->CameraDepth != nullptr)
 			{
-				// Only create a new camera depth texture if it's set and we don't already have one
-				if (CameraDepth == nullptr)
-				{
-					CameraDepth = NewObject<UAppleARKitTextureCameraDepth>();
-				}
 				// Reuse the UObjects because otherwise the time between GCs causes ARKit to be starved of resources
                 CameraDepth->Init(FPlatformTime::Seconds(), GameThreadFrame->CameraDepth);
 			}
@@ -1178,11 +1155,30 @@ void FAppleARKitSystem::ClearTrackedGeometries()
 #endif
 }
 
+void FAppleARKitSystem::SetupCameraTextures()
+{
+#if SUPPORTS_ARKIT_1_0
+	if (CameraImage == nullptr)
+	{
+		CameraImage = NewObject<UAppleARKitTextureCameraImage>();
+		CameraImage->Init(FPlatformTime::Seconds(), nullptr);
+		FAppleARKitXRCamera* Camera = (FAppleARKitXRCamera*)XRCamera.Get();
+		Camera->SetOverlayTexture(CameraImage);
+	}
+	if (CameraDepth == nullptr)
+	{
+		CameraDepth = NewObject<UAppleARKitTextureCameraDepth>();
+	}
+#endif
+}
+
 PRAGMA_DISABLE_OPTIMIZATION
 bool FAppleARKitSystem::Run(UARSessionConfig* SessionConfig)
 {
 	TimecodeProvider = UAppleARKitSettings::GetTimecodeProvider();
 
+	SetupCameraTextures();
+
 	{
 		// Clear out any existing frames since they aren't valid anymore
 		FScopeLock ScopeLock(&FrameLock);
diff --git a/Engine/Plugins/Runtime/AR/Apple/AppleARKit/Source/AppleARKit/Private/AppleARKitSystem.h b/Engine/Plugins/Runtime/AR/Apple/AppleARKit/Source/AppleARKit/Private/AppleARKitSystem.h
--- a/Engine/Plugins/Runtime/AR/Apple/AppleARKit/Source/AppleARKit/Private/AppleARKitSystem.h
+++ b/Engine/Plugins/Runtime/AR/Apple/AppleARKit/Source/AppleARKit/Private/AppleARKitSystem.h
@@ -146,7 +146,10 @@ private:
 	
 	/** Updates the ARKit perf counters */
 	void UpdateARKitPerfStats();
-	
+
+	/** Inits the textures and sets the texture on the overlay */
+	void SetupCameraTextures();
+
 	/** The orientation of the device; see EDeviceScreenOrientation */
 	EDeviceScreenOrientation DeviceOrientation;
 	
diff --git a/Engine/Plugins/Runtime/AR/Apple/AppleARKit/Source/AppleARKit/Private/AppleARKitTextures.cpp b/Engine/Plugins/Runtime/AR/Apple/AppleARKit/Source/AppleARKit/Private/AppleARKitTextures.cpp
--- a/Engine/Plugins/Runtime/AR/Apple/AppleARKit/Source/AppleARKit/Private/AppleARKitTextures.cpp
+++ b/Engine/Plugins/Runtime/AR/Apple/AppleARKit/Source/AppleARKit/Private/AppleARKitTextures.cpp
@@ -45,22 +45,31 @@ public:
 			CGColorSpaceRef ColorSpaceRef = CGColorSpaceCreateDeviceRGB();
 			CIImage* Image = [[CIImage alloc] initWithCVPixelBuffer: CameraImage];
 
-			// Textures always need to be rotated to the "LandscapeRight" orientation
+			// Textures always need to be rotated so to a sane orientation (and mirrored because of differing coord system)
 			CIImage* RotatedImage = [Image imageByApplyingOrientation: GetRotationFromDeviceOrientation()];
+			// Get the sizes from the rotated image
+			CGRect ImageExtent = RotatedImage.extent;
 
-			// Get the size using the rotated image
-			Size.X = RotatedImage.extent.size.width;
-			Size.Y = RotatedImage.extent.size.height;
-
-			// Create the target texture that we'll update into
-			FRHIResourceCreateInfo CreateInfo;
-			DecodedTextureRef = RHICreateTexture2D(Size.X, Size.Y, PF_B8G8R8A8, 1, 1, TexCreate_Dynamic | TexCreate_ShaderResource | TexCreate_UAV, CreateInfo);
+			// Don't reallocate the texture if the sizes match
+			if (Size.X != ImageExtent.size.width || Size.Y != ImageExtent.size.height)
+			{
+				Size.X = ImageExtent.size.width;
+				Size.Y = ImageExtent.size.height;
+				
+				// Let go of the last texture
+				RHIUpdateTextureReference(Owner->TextureReference.TextureReferenceRHI, FTextureRHIParamRef());
+				DecodedTextureRef.SafeRelease();
+				
+				// Create the target texture that we'll update into
+				FRHIResourceCreateInfo CreateInfo;
+				DecodedTextureRef = RHICreateTexture2D(Size.X, Size.Y, PF_B8G8R8A8, 1, 1, TexCreate_Dynamic | TexCreate_ShaderResource | TexCreate_UAV, CreateInfo);
+			}
 
 			// Get the underlying metal texture so we can render to it
 			id<MTLTexture> UnderlyingMetalTexture = (id<MTLTexture>)DecodedTextureRef->GetNativeResource();
 
 			// Do the conversion on the GPU
-			[[CIContext context] render: RotatedImage toMTLTexture: UnderlyingMetalTexture commandBuffer: nil bounds: RotatedImage.extent colorSpace: ColorSpaceRef];
+			[[CIContext context] render: RotatedImage toMTLTexture: UnderlyingMetalTexture commandBuffer: nil bounds: ImageExtent colorSpace: ColorSpaceRef];
 
 			// Now that the conversion is done, we can get rid of our refs
 			[Image release];
@@ -119,7 +128,6 @@ public:
 		if (LastFrameNumber != GFrameNumber)
 		{
 			LastFrameNumber = GFrameNumber;
-			ReleaseRHI();
 			CameraImage = InCameraImage;
 			CFRetain(CameraImage);
 			InitRHI();
@@ -132,27 +140,28 @@ private:
 	/** @return the rotation to use to rotate the texture to the proper direction */
 	int32 GetRotationFromDeviceOrientation()
 	{
-		extern UIInterfaceOrientation GInterfaceOrientation;
-		switch (GInterfaceOrientation)
+		// NOTE: The texture we are reading from is in device space and mirrored, because Apple hates us
+		EDeviceScreenOrientation ScreenOrientation = FPlatformMisc::GetDeviceOrientation();
+		switch (ScreenOrientation)
 		{
-			case UIInterfaceOrientationPortrait:
+			case EDeviceScreenOrientation::Portrait:
 			{
-				return kCGImagePropertyOrientationRight;
+				return kCGImagePropertyOrientationRightMirrored;
 			}
 
-			case UIInterfaceOrientationLandscapeLeft:
+			case EDeviceScreenOrientation::LandscapeLeft:
 			{
-				return kCGImagePropertyOrientationDown;
+				return kCGImagePropertyOrientationUpMirrored;
 			}
 
-			case UIInterfaceOrientationPortraitUpsideDown:
+			case EDeviceScreenOrientation::PortraitUpsideDown:
 			{
-				return kCGImagePropertyOrientationLeft;
+				return kCGImagePropertyOrientationLeftMirrored;
 			}
 
-			case UIInterfaceOrientationLandscapeRight:
+			case EDeviceScreenOrientation::LandscapeRight:
 			{
-				return kCGImagePropertyOrientationUp;
+				return kCGImagePropertyOrientationDownMirrored;
 			}
 		}
 		// Don't know so don't rotate
diff --git a/Engine/Plugins/Runtime/AR/Apple/AppleARKit/Source/AppleARKit/Private/AppleARKitVideoOverlay.cpp b/Engine/Plugins/Runtime/AR/Apple/AppleARKit/Source/AppleARKit/Private/AppleARKitVideoOverlay.cpp
--- a/Engine/Plugins/Runtime/AR/Apple/AppleARKit/Source/AppleARKit/Private/AppleARKitVideoOverlay.cpp
+++ b/Engine/Plugins/Runtime/AR/Apple/AppleARKit/Source/AppleARKit/Private/AppleARKitVideoOverlay.cpp
@@ -3,6 +3,7 @@
 #include "AppleARKitVideoOverlay.h"
 #include "AppleARKitFrame.h"
 #include "Materials/Material.h"
+#include "Materials/MaterialInstanceDynamic.h"
 #include "MaterialShaderType.h"
 #include "MaterialShader.h"
 #include "ExternalTexture.h"
@@ -25,234 +26,16 @@
 	#include "IOSAppDelegate.h"
 #endif
 
-#if SUPPORTS_ARKIT_1_0
-/**
-* Passes a CVMetalTextureRef through to the RHI to wrap in an RHI texture without traversing system memory.
-* @see FAvfTexture2DResourceWrapper & FMetalSurface::FMetalSurface
-*/
-class FAppleARKitCameraTextureResourceWrapper : public FResourceBulkDataInterface
-{
-public:
-	FAppleARKitCameraTextureResourceWrapper(CFTypeRef InImageBuffer)
-		: ImageBuffer(InImageBuffer)
-	{
-		check(ImageBuffer);
-		CFRetain(ImageBuffer);
-	}
-
-	/**
-	* @return ptr to the resource memory which has been preallocated
-	*/
-	virtual const void* GetResourceBulkData() const override
-	{
-		return ImageBuffer;
-	}
-
-	/**
-	* @return size of resource memory
-	*/
-	virtual uint32 GetResourceBulkDataSize() const override
-	{
-		return 0;
-	}
-
-	/**
-	* @return the type of bulk data for special handling
-	*/
-	virtual EBulkDataType GetResourceType() const override
-	{
-		return EBulkDataType::MediaTexture;
-	}
-
-	/**
-	* Free memory after it has been used to initialize RHI resource
-	*/
-	virtual void Discard() override
-	{
-		delete this;
-	}
-
-	virtual ~FAppleARKitCameraTextureResourceWrapper()
-	{
-		CFRelease(ImageBuffer);
-		ImageBuffer = nullptr;
-	}
-
-	CFTypeRef ImageBuffer;
-};
-#endif
-
 FAppleARKitVideoOverlay::FAppleARKitVideoOverlay()
-	: RenderingOverlayMaterial(nullptr)
-	, LastUpdateTimestamp(-1.0)
+	: MID_CameraOverlay(nullptr)
+	, RenderingOverlayMaterial(nullptr)
 {
-	RenderingOverlayMaterial = LoadObject<UMaterialInterface>(nullptr, TEXT("/AppleARKit/ARKitCameraMaterial.ARKitCameraMaterial"));
+	UMaterialInterface* LoadedMaterial = LoadObject<UMaterialInterface>(nullptr, TEXT("/AppleARKit/M_CameraOverlay.M_CameraOverlay"));
+	check(LoadedMaterial != nullptr);
+	// Create a MaterialInstanceDynamic which we'll replace the texture with the camera texture from ARKit
+	MID_CameraOverlay = UMaterialInstanceDynamic::Create(LoadedMaterial, GetTransientPackage());
+	RenderingOverlayMaterial = MID_CameraOverlay;
 	check(RenderingOverlayMaterial != nullptr);
-	RenderingOverlayMaterial->AddToRoot();
-}
-
-void FAppleARKitVideoOverlay::UpdateVideoTexture_RenderThread(FRHICommandListImmediate& RHICmdList, FAppleARKitFrame& Frame, const FSceneViewFamily& InViewFamily)
-{
-	// Allocate and register
-	if (VideoTextureY == nullptr)
-	{
-		check(VideoTextureCbCr == nullptr);
-		check(OverlayIndexBufferRHI == nullptr);
-		check(OverlayVertexBufferRHI[0] == nullptr);
-
-		FRHIResourceCreateInfo CreateInfo;
-		VideoTextureY = RHICmdList.CreateTexture2D(1, 1, PF_R8G8B8A8, 1, 1, 0, CreateInfo);
-		VideoTextureCbCr = RHICmdList.CreateTexture2D(1, 1, PF_R8G8B8A8, 1, 1, 0, CreateInfo);
-
-		FSamplerStateInitializerRHI SamplerStateInitializer(SF_Bilinear, AM_Wrap, AM_Wrap, AM_Wrap);
-		FSamplerStateRHIRef SamplerStateRHI = RHICreateSamplerState(SamplerStateInitializer);
-
-		FExternalTextureRegistry::Get().RegisterExternalTexture(ARKitPassthroughCameraExternalTextureYGuid, VideoTextureY, SamplerStateRHI);
-		FExternalTextureRegistry::Get().RegisterExternalTexture(ARKitPassthroughCameraExternalTextureCbCrGuid, VideoTextureCbCr, SamplerStateRHI);
-
-		// Setup index buffer
-		const uint16 Indices[] = { 0, 1, 2, 2, 1, 3 };
-
-		TResourceArray<uint16, INDEXBUFFER_ALIGNMENT> IndexBuffer;
-		const uint32 NumIndices = ARRAY_COUNT(Indices);
-		IndexBuffer.AddUninitialized(NumIndices);
-		FMemory::Memcpy(IndexBuffer.GetData(), Indices, NumIndices * sizeof(uint16));
-
-		FRHIResourceCreateInfo CreateInfoIB(&IndexBuffer);
-		OverlayIndexBufferRHI = RHICreateIndexBuffer(sizeof(uint16), IndexBuffer.GetResourceDataSize(), BUF_Static, CreateInfoIB);
-
-		check(InViewFamily.Views.Num() > 0);
-		const FSceneView& View = *InViewFamily.Views[0];
-
-		const FVector2D ViewSize(View.UnconstrainedViewRect.Max.X, View.UnconstrainedViewRect.Max.Y);
-		
-		// CameraSize is 1280 x 720 regardless of the device orientation. We flip it here if needed to make it consistent with the view size.
-		FVector2D CameraSize = Frame.Camera.ImageResolution;
-		if ((ViewSize.X > ViewSize.Y) != (CameraSize.X > CameraSize.Y))
-		{
-			CameraSize = FVector2D(CameraSize.Y, CameraSize.X);
-		}
-		
-		const float CameraAspectRatio = CameraSize.X / CameraSize.Y;
-		const float ViewAspectRatio = ViewSize.X / ViewSize.Y;
-		const float ViewAspectRatioLandscape = (ViewSize.X > ViewSize.Y) ? ViewAspectRatio : ViewSize.Y / ViewSize.X;
-		
-		float UVOffsetAmount = 0.0f;
-		if (!FMath::IsNearlyEqual(ViewAspectRatio, CameraAspectRatio))
-		{
-			if (ViewAspectRatio > CameraAspectRatio)
-			{
-				UVOffsetAmount = 0.5f * (1.0f - (CameraAspectRatio / ViewAspectRatio));
-			}
-			else
-			{
-				UVOffsetAmount = 0.5f * (1.0f - (ViewAspectRatio / CameraAspectRatio));
-			}
-		}
-		
-		UVOffset = (ViewAspectRatioLandscape <= Frame.Camera.GetAspectRatio()) ? FVector2D(UVOffsetAmount, 0.0f) : FVector2D(0.0f, UVOffsetAmount);
-		
-		// Setup vertex buffer
-		const FVector4 Positions[] =
-		{
-			FVector4(0.0f, 1.0f, 0.0f, 1.0f),
-			FVector4(0.0f, 0.0f, 0.0f, 1.0f),
-			FVector4(1.0f, 1.0f, 0.0f, 1.0f),
-			FVector4(1.0f, 0.0f, 0.0f, 1.0f)
-		};
-
-		const FVector2D UVs[] =
-		{
-			// Landscape left
-			FVector2D(UVOffset.X, 1.0f - UVOffset.Y),
-			FVector2D(UVOffset.X, UVOffset.Y),
-			FVector2D(1.0f - UVOffset.X, 1.0f - UVOffset.Y),
-			FVector2D(1.0f - UVOffset.X, UVOffset.Y),
-
-			// Landscape right
-			FVector2D(1.0f - UVOffset.X, UVOffset.Y),
-			FVector2D(1.0f - UVOffset.X, 1.0f - UVOffset.Y),
-			FVector2D(UVOffset.X, UVOffset.Y),
-			FVector2D(UVOffset.X, 1.0f - UVOffset.Y),
-
-			// Portrait
-			FVector2D(1.0f - UVOffset.X, 1.0f - UVOffset.Y),
-			FVector2D(UVOffset.X, 1.0f - UVOffset.Y),
-			FVector2D(1.0f - UVOffset.X, UVOffset.Y),
-			FVector2D(UVOffset.X, UVOffset.Y),
-
-			// Portrait Upside Down
-			FVector2D(UVOffset.X, UVOffset.Y),
-			FVector2D(1.0f - UVOffset.X, UVOffset.Y),
-			FVector2D(UVOffset.X, 1.0f - UVOffset.Y),
-			FVector2D(1.0f - UVOffset.X, 1.0f - UVOffset.Y)
-		};
-
-		uint32 UVIndex = 0;
-		for (uint32 OrientationIter = 0; OrientationIter < 4; ++OrientationIter)
-		{
-			TResourceArray<FFilterVertex, VERTEXBUFFER_ALIGNMENT> Vertices;
-			Vertices.SetNumUninitialized(4);
-
-			Vertices[0].Position = Positions[0];
-			Vertices[0].UV = UVs[UVIndex];
-
-			Vertices[1].Position = Positions[1];
-			Vertices[1].UV = UVs[UVIndex + 1];
-
-			Vertices[2].Position = Positions[2];
-			Vertices[2].UV = UVs[UVIndex + 2];
-
-			Vertices[3].Position = Positions[3];
-			Vertices[3].UV = UVs[UVIndex + 3];
-
-			UVIndex += 4;
-
-			FRHIResourceCreateInfo CreateInfoVB(&Vertices);
-			OverlayVertexBufferRHI[OrientationIter] = RHICreateVertexBuffer(Vertices.GetResourceDataSize(), BUF_Static, CreateInfoVB);
-		}
-	}
-
-#if SUPPORTS_ARKIT_1_0
-	if (FAppleARKitAvailability::SupportsARKit10())
-	{
-		check(IsMetalPlatform(GMaxRHIShaderPlatform));
-
-		if (LastUpdateTimestamp != Frame.Timestamp && Frame.CapturedYImage && Frame.CapturedCbCrImage)
-		{
-			FRHIResourceCreateInfo CreateInfo;
-			const uint32 CreateFlags = TexCreate_Dynamic | TexCreate_NoTiling | TexCreate_ShaderResource;
-			CreateInfo.BulkData = new FAppleARKitCameraTextureResourceWrapper(Frame.CapturedYImage);
-			CreateInfo.ResourceArray = nullptr;
-
-			// pull the Y and CbCr textures out of the captured image planes (format is fake here, it will get the format from the FAppleARKitCameraTextureResourceWrapper)
-			VideoTextureY = RHICreateTexture2D(Frame.CapturedYImageWidth, Frame.CapturedYImageHeight, /*Format=*/PF_B8G8R8A8, /*NumMips=*/1, /*NumSamples=*/1, CreateFlags, CreateInfo);
-
-			CreateInfo.BulkData = new FAppleARKitCameraTextureResourceWrapper(Frame.CapturedCbCrImage);
-			VideoTextureCbCr = RHICreateTexture2D(Frame.CapturedCbCrImageWidth, Frame.CapturedCbCrImageHeight, /*Format=*/PF_B8G8R8A8, /*NumMips=*/1, /*NumSamples=*/1, CreateFlags, CreateInfo);
-
-			// todo: Add an update call to the registry instead of this unregister/re-register
-			FExternalTextureRegistry::Get().UnregisterExternalTexture(ARKitPassthroughCameraExternalTextureYGuid);
-			FExternalTextureRegistry::Get().UnregisterExternalTexture(ARKitPassthroughCameraExternalTextureCbCrGuid);
-
-			FSamplerStateInitializerRHI SamplerStateInitializer(SF_Bilinear, AM_Wrap, AM_Wrap, AM_Wrap);
-			FSamplerStateRHIRef SamplerStateRHI = RHICreateSamplerState(SamplerStateInitializer);
-
-			FExternalTextureRegistry::Get().RegisterExternalTexture(ARKitPassthroughCameraExternalTextureYGuid, VideoTextureY, SamplerStateRHI);
-			FExternalTextureRegistry::Get().RegisterExternalTexture(ARKitPassthroughCameraExternalTextureCbCrGuid, VideoTextureCbCr, SamplerStateRHI);
-			
-			//Make sure AR camera pass through materials are updated properly
-			FMaterialRenderProxy::UpdateDeferredCachedUniformExpressions();
-
-			CFRelease(Frame.CapturedYImage);
-			CFRelease(Frame.CapturedCbCrImage);
-			Frame.CapturedYImage = nullptr;
-			Frame.CapturedCbCrImage = nullptr;
-
-			LastUpdateTimestamp = Frame.Timestamp;
-		}
-	}
-#endif
 }
 
 // We use something similar to the PostProcessMaterial to render the color camera overlay.
@@ -375,162 +158,139 @@ private:
 IMPLEMENT_MATERIAL_SHADER_TYPE(template<>, FARKitCameraOverlayPS<true>, TEXT("/Engine/Private/PostProcessMaterialShaders.usf"), TEXT("MainPS_ES2"), SF_Pixel);
 IMPLEMENT_MATERIAL_SHADER_TYPE(template<>, FARKitCameraOverlayPS<false>, TEXT("/Engine/Private/PostProcessMaterialShaders.usf"), TEXT("MainPS_VideoOverlay"), SF_Pixel);
 
-void FAppleARKitVideoOverlay::RenderVideoOverlay_RenderThread(FRHICommandListImmediate& RHICmdList, const FSceneView& InView, const EDeviceScreenOrientation DeviceOrientation)
+void FAppleARKitVideoOverlay::RenderVideoOverlay_RenderThread(FRHICommandListImmediate& RHICmdList, const FSceneView& InView)
 {
-#if SUPPORTS_ARKIT_1_0
+	if (RenderingOverlayMaterial == nullptr || !RenderingOverlayMaterial->IsValidLowLevel())
+	{
+		return;
+	}
 
-	if (FAppleARKitAvailability::SupportsARKit10())
+	if (VertexBufferRHI == nullptr || !VertexBufferRHI.IsValid())
 	{
-		if (VideoTextureY == nullptr || VideoTextureCbCr == nullptr || RenderingOverlayMaterial == nullptr || !RenderingOverlayMaterial->IsValidLowLevel())
-		{
-			return;
-		}
+		// Setup vertex buffer
+		TResourceArray<FFilterVertex, VERTEXBUFFER_ALIGNMENT> Vertices;
+		Vertices.SetNumUninitialized(4);
 
-		const auto FeatureLevel = InView.GetFeatureLevel();
-		IRendererModule& RendererModule = GetRendererModule();
+		Vertices[0].Position = FVector4(0.f, 0.f, 0.f, 1.f);
+		Vertices[0].UV = FVector2D(0.f, 0.f);
 
-		const FMaterial* const CameraMaterial = RenderingOverlayMaterial->GetRenderProxy()->GetMaterial(FeatureLevel);
-		const FMaterialShaderMap* const MaterialShaderMap = CameraMaterial->GetRenderingThreadShaderMap();
+		Vertices[1].Position = FVector4(1.f, 0.f, 0.f, 1.f);
+		Vertices[1].UV = FVector2D(1.f, 0.f);
 
-		FGraphicsPipelineStateInitializer GraphicsPSOInit;
-		RHICmdList.ApplyCachedRenderTargets(GraphicsPSOInit);
+		Vertices[2].Position = FVector4(1.f, 1.f, 0.f, 1.f);
+		Vertices[2].UV = FVector2D(1.f, 1.f);
 
-		GraphicsPSOInit.BlendState = TStaticBlendState<>::GetRHI();
-		GraphicsPSOInit.RasterizerState = TStaticRasterizerState<>::GetRHI();
-		GraphicsPSOInit.DepthStencilState = TStaticDepthStencilState<false, CF_DepthNearOrEqual>::GetRHI();
-		GraphicsPSOInit.PrimitiveType = PT_TriangleList;
-		GraphicsPSOInit.BoundShaderState.VertexDeclarationRHI = GFilterVertexDeclaration.VertexDeclarationRHI;
-		
-		const bool bIsMobileRenderer = FeatureLevel <= ERHIFeatureLevel::ES3_1;
-		FMaterialShader* VertexShader = nullptr;
-		FMaterialShader* PixelShader = nullptr;
-		if (bIsMobileRenderer)
-		{
-			VertexShader = MaterialShaderMap->GetShader<FARKitCameraOverlayVS<true>>();
-			PixelShader = MaterialShaderMap->GetShader<FARKitCameraOverlayPS<true>>();
-			GraphicsPSOInit.BoundShaderState.VertexShaderRHI = GETSAFERHISHADER_VERTEX(reinterpret_cast<FARKitCameraOverlayVS<true>*>(VertexShader));
-			GraphicsPSOInit.BoundShaderState.PixelShaderRHI = GETSAFERHISHADER_PIXEL(reinterpret_cast<FARKitCameraOverlayPS<true>*>(PixelShader));
-		}
-		else
-		{
-			VertexShader = MaterialShaderMap->GetShader<FARKitCameraOverlayVS<false>>();
-			PixelShader = MaterialShaderMap->GetShader<FARKitCameraOverlayPS<false>>();
-			GraphicsPSOInit.BoundShaderState.VertexShaderRHI = GETSAFERHISHADER_VERTEX(reinterpret_cast<FARKitCameraOverlayVS<false>*>(VertexShader));
-			GraphicsPSOInit.BoundShaderState.PixelShaderRHI = GETSAFERHISHADER_PIXEL(reinterpret_cast<FARKitCameraOverlayPS<false>*>(PixelShader));
-		}
-		
-		SetGraphicsPipelineState(RHICmdList, GraphicsPSOInit);
-		
-		const FIntPoint ViewSize = InView.UnconstrainedViewRect.Size();
-		FDrawRectangleParameters Parameters;
-		Parameters.PosScaleBias = FVector4(ViewSize.X, ViewSize.Y, 0, 0);
-		Parameters.UVScaleBias = FVector4(1.0f, 1.0f, 0.0f, 0.0f);
-		Parameters.InvTargetSizeAndTextureSize = FVector4(
-														  1.0f / ViewSize.X, 1.0f / ViewSize.Y,
-														  1.0f, 1.0f);
+		Vertices[3].Position = FVector4(0.f, 1.f, 0.f, 1.f);
+		Vertices[3].UV = FVector2D(0.f, 1.f);
 
-		if (bIsMobileRenderer)
-		{
-			FARKitCameraOverlayVS<true>* const VertexShaderPtr = reinterpret_cast<FARKitCameraOverlayVS<true>*>(VertexShader);
-			SetUniformBufferParameterImmediate(RHICmdList, VertexShaderPtr->GetVertexShader(), VertexShaderPtr->GetUniformBufferParameter<FDrawRectangleParameters>(), Parameters);
-			VertexShaderPtr->SetParameters(RHICmdList, InView);
-			reinterpret_cast<FARKitCameraOverlayPS<true>*>(PixelShader)->SetParameters(RHICmdList, InView, RenderingOverlayMaterial->GetRenderProxy());
-		}
-		else
-		{
-			FARKitCameraOverlayVS<false>* const VertexShaderPtr = reinterpret_cast<FARKitCameraOverlayVS<false>*>(VertexShader);
-			SetUniformBufferParameterImmediate(RHICmdList, VertexShaderPtr->GetVertexShader(), VertexShaderPtr->GetUniformBufferParameter<FDrawRectangleParameters>(), Parameters);
-			VertexShaderPtr->SetParameters(RHICmdList, InView);
-			reinterpret_cast<FARKitCameraOverlayPS<false>*>(PixelShader)->SetParameters(RHICmdList, InView, RenderingOverlayMaterial->GetRenderProxy());
-		}
+		FRHIResourceCreateInfo CreateInfoVB(&Vertices);
+		VertexBufferRHI = RHICreateVertexBuffer(Vertices.GetResourceDataSize(), BUF_Static, CreateInfoVB);
+	}
 
-		FVertexBufferRHIParamRef VertexBufferRHI = nullptr;
-		switch (DeviceOrientation)
-		{
-			case EDeviceScreenOrientation::LandscapeRight:
-				VertexBufferRHI = OverlayVertexBufferRHI[0];
-				break;
+	if (IndexBufferRHI == nullptr || !IndexBufferRHI.IsValid())
+	{
+		// Setup index buffer
+		const uint16 Indices[] = { 0, 1, 2, 2, 3, 0 };
+
+		TResourceArray<uint16, INDEXBUFFER_ALIGNMENT> IndexBuffer;
+		const uint32 NumIndices = ARRAY_COUNT(Indices);
+		IndexBuffer.AddUninitialized(NumIndices);
+		FMemory::Memcpy(IndexBuffer.GetData(), Indices, NumIndices * sizeof(uint16));
 
-			case EDeviceScreenOrientation::LandscapeLeft:
-				VertexBufferRHI = OverlayVertexBufferRHI[1];
-				break;
+		FRHIResourceCreateInfo CreateInfoIB(&IndexBuffer);
+		IndexBufferRHI = RHICreateIndexBuffer(sizeof(uint16), IndexBuffer.GetResourceDataSize(), BUF_Static, CreateInfoIB);
+	}
 
-			case EDeviceScreenOrientation::Portrait:
-				VertexBufferRHI = OverlayVertexBufferRHI[2];
-				break;
+	const auto FeatureLevel = InView.GetFeatureLevel();
+	IRendererModule& RendererModule = GetRendererModule();
 
-			case EDeviceScreenOrientation::PortraitUpsideDown:
-				VertexBufferRHI = OverlayVertexBufferRHI[3];
-				break;
+	const FMaterial* const CameraMaterial = RenderingOverlayMaterial->GetRenderProxy()->GetMaterial(FeatureLevel);
+	const FMaterialShaderMap* const MaterialShaderMap = CameraMaterial->GetRenderingThreadShaderMap();
 
-			default:
-				VertexBufferRHI = OverlayVertexBufferRHI[0];
-				break;
-		}
+	FGraphicsPipelineStateInitializer GraphicsPSOInit;
+	RHICmdList.ApplyCachedRenderTargets(GraphicsPSOInit);
 
-		if (VertexBufferRHI && OverlayIndexBufferRHI.IsValid())
-		{
-			RHICmdList.SetStreamSource(0, VertexBufferRHI, 0);
-			RHICmdList.DrawIndexedPrimitive(
-				OverlayIndexBufferRHI,
-				/*BaseVertexIndex=*/ 0,
-				/*MinIndex=*/ 0,
-				/*NumVertices=*/ 4,
-				/*StartIndex=*/ 0,
-				/*NumPrimitives=*/ 2,
-				/*NumInstances=*/ 1
-			);
-		}
+	GraphicsPSOInit.BlendState = TStaticBlendState<>::GetRHI();
+	GraphicsPSOInit.RasterizerState = TStaticRasterizerState<>::GetRHI();
+	GraphicsPSOInit.DepthStencilState = TStaticDepthStencilState<false, CF_DepthNearOrEqual>::GetRHI();
+	GraphicsPSOInit.PrimitiveType = PT_TriangleList;
+	GraphicsPSOInit.BoundShaderState.VertexDeclarationRHI = GFilterVertexDeclaration.VertexDeclarationRHI;
+
+	const bool bIsMobileRenderer = FeatureLevel <= ERHIFeatureLevel::ES3_1;
+	FMaterialShader* VertexShader = nullptr;
+	FMaterialShader* PixelShader = nullptr;
+	if (bIsMobileRenderer)
+	{
+		VertexShader = MaterialShaderMap->GetShader<FARKitCameraOverlayVS<true>>();
+		PixelShader = MaterialShaderMap->GetShader<FARKitCameraOverlayPS<true>>();
+		GraphicsPSOInit.BoundShaderState.VertexShaderRHI = GETSAFERHISHADER_VERTEX(reinterpret_cast<FARKitCameraOverlayVS<true>*>(VertexShader));
+		GraphicsPSOInit.BoundShaderState.PixelShaderRHI = GETSAFERHISHADER_PIXEL(reinterpret_cast<FARKitCameraOverlayPS<true>*>(PixelShader));
+	}
+	else
+	{
+		VertexShader = MaterialShaderMap->GetShader<FARKitCameraOverlayVS<false>>();
+		PixelShader = MaterialShaderMap->GetShader<FARKitCameraOverlayPS<false>>();
+		GraphicsPSOInit.BoundShaderState.VertexShaderRHI = GETSAFERHISHADER_VERTEX(reinterpret_cast<FARKitCameraOverlayVS<false>*>(VertexShader));
+		GraphicsPSOInit.BoundShaderState.PixelShaderRHI = GETSAFERHISHADER_PIXEL(reinterpret_cast<FARKitCameraOverlayPS<false>*>(PixelShader));
 	}
-#endif
-}
 
+	SetGraphicsPipelineState(RHICmdList, GraphicsPSOInit);
 
-bool FAppleARKitVideoOverlay::GetPassthroughCameraUVs_RenderThread(TArray<FVector2D>& OutUVs, const EDeviceScreenOrientation DeviceOrientation)
-{
-#if SUPPORTS_ARKIT_1_0
-	if (VideoTextureY != nullptr) // make sure UVOffset has been initialized
-	{
-		OutUVs.SetNumUninitialized(4);
+	const FIntPoint ViewSize = InView.UnconstrainedViewRect.Size();
+	FDrawRectangleParameters Parameters;
+	Parameters.PosScaleBias = FVector4(ViewSize.X, ViewSize.Y, 0, 0);
+	Parameters.UVScaleBias = FVector4(1.0f, 1.0f, 0.0f, 0.0f);
+	Parameters.InvTargetSizeAndTextureSize = FVector4(
+													  1.0f / ViewSize.X, 1.0f / ViewSize.Y,
+													  1.0f, 1.0f);
 
-		switch (DeviceOrientation)
-		{
-		case EDeviceScreenOrientation::LandscapeRight:
-			OutUVs[1] = FVector2D(UVOffset.X, 1.0f - UVOffset.Y);
-			OutUVs[0] = FVector2D(UVOffset.X, UVOffset.Y);
-			OutUVs[3] = FVector2D(1.0f - UVOffset.X, 1.0f - UVOffset.Y);
-			OutUVs[2] = FVector2D(1.0f - UVOffset.X, UVOffset.Y);
-			return true;
-
-		case EDeviceScreenOrientation::LandscapeLeft:
-			OutUVs[1] = FVector2D(1.0f - UVOffset.X, UVOffset.Y);
-			OutUVs[0] = FVector2D(1.0f - UVOffset.X, 1.0f - UVOffset.Y);
-			OutUVs[3] = FVector2D(UVOffset.X, UVOffset.Y);
-			OutUVs[2] = FVector2D(UVOffset.X, 1.0f - UVOffset.Y);
-			return true;
-
-		case EDeviceScreenOrientation::Portrait:
-			OutUVs[1] = FVector2D(1.0f - UVOffset.X, 1.0f - UVOffset.Y);
-			OutUVs[0] = FVector2D(UVOffset.X, 1.0f - UVOffset.Y);
-			OutUVs[3] = FVector2D(1.0f - UVOffset.X, UVOffset.Y);
-			OutUVs[2] = FVector2D(UVOffset.X, UVOffset.Y);
-			return true;
-
-		case EDeviceScreenOrientation::PortraitUpsideDown:
-			OutUVs[1] = FVector2D(UVOffset.X, UVOffset.Y);
-			OutUVs[0] = FVector2D(1.0f - UVOffset.X, UVOffset.Y);
-			OutUVs[3] = FVector2D(UVOffset.X, 1.0f - UVOffset.Y);
-			OutUVs[2] = FVector2D(1.0f - UVOffset.X, 1.0f - UVOffset.Y);
-			return true;
-
-		default:
-			return false;
-		}
+	if (bIsMobileRenderer)
+	{
+		FARKitCameraOverlayVS<true>* const VertexShaderPtr = reinterpret_cast<FARKitCameraOverlayVS<true>*>(VertexShader);
+		SetUniformBufferParameterImmediate(RHICmdList, VertexShaderPtr->GetVertexShader(), VertexShaderPtr->GetUniformBufferParameter<FDrawRectangleParameters>(), Parameters);
+		VertexShaderPtr->SetParameters(RHICmdList, InView);
+		reinterpret_cast<FARKitCameraOverlayPS<true>*>(PixelShader)->SetParameters(RHICmdList, InView, RenderingOverlayMaterial->GetRenderProxy());
 	}
 	else
-#endif
 	{
-		return false;
+		FARKitCameraOverlayVS<false>* const VertexShaderPtr = reinterpret_cast<FARKitCameraOverlayVS<false>*>(VertexShader);
+		SetUniformBufferParameterImmediate(RHICmdList, VertexShaderPtr->GetVertexShader(), VertexShaderPtr->GetUniformBufferParameter<FDrawRectangleParameters>(), Parameters);
+		VertexShaderPtr->SetParameters(RHICmdList, InView);
+		reinterpret_cast<FARKitCameraOverlayPS<false>*>(PixelShader)->SetParameters(RHICmdList, InView, RenderingOverlayMaterial->GetRenderProxy());
 	}
+
+	RHICmdList.SetStreamSource(0, VertexBufferRHI, 0);
+	RHICmdList.DrawIndexedPrimitive(
+		IndexBufferRHI,
+		/*BaseVertexIndex=*/ 0,
+		/*MinIndex=*/ 0,
+		/*NumVertices=*/ 4,
+		/*StartIndex=*/ 0,
+		/*NumPrimitives=*/ 2,
+		/*NumInstances=*/ 1
+	);
 }
 
+
+bool FAppleARKitVideoOverlay::GetPassthroughCameraUVs_RenderThread(TArray<FVector2D>& OutUVs, const EDeviceScreenOrientation)
+{
+	OutUVs.Reset(4);
+
+	new(OutUVs) FVector2D(0.f, 0.f);
+	new(OutUVs) FVector2D(1.f, 0.f);
+	new(OutUVs) FVector2D(1.0f, 1.0f);
+	new(OutUVs) FVector2D(0.f, 1.0f);
+
+	return true;
+}
+
+void FAppleARKitVideoOverlay::AddReferencedObjects(FReferenceCollector& Collector)
+{
+	Collector.AddReferencedObject(RenderingOverlayMaterial);
+	Collector.AddReferencedObject(MID_CameraOverlay);
+}
+
+void FAppleARKitVideoOverlay::SetOverlayTexture(UARTextureCameraImage* InCameraImage)
+{
+	MID_CameraOverlay->SetTextureParameterValue(FName("CameraImage"), InCameraImage);
+}
diff --git a/Engine/Plugins/Runtime/AR/Apple/AppleARKit/Source/AppleARKit/Private/AppleARKitVideoOverlay.h b/Engine/Plugins/Runtime/AR/Apple/AppleARKit/Source/AppleARKit/Private/AppleARKitVideoOverlay.h
--- a/Engine/Plugins/Runtime/AR/Apple/AppleARKit/Source/AppleARKit/Private/AppleARKitVideoOverlay.h
+++ b/Engine/Plugins/Runtime/AR/Apple/AppleARKit/Source/AppleARKit/Private/AppleARKitVideoOverlay.h
@@ -5,10 +5,10 @@
 #include "CoreMinimal.h"
 #include "Materials/MaterialInterface.h"
 #include "UObject/ConstructorHelpers.h"
-#include "Kismet/BlueprintPlatformLibrary.h"
 #include "AppleARKitVideoOverlay.generated.h"
 
-class FSceneViewFamily;
+class UARTextureCameraImage;
+class UMaterialInstanceDynamic;
 
 /** Helper class to ensure the ARKit camera material is cooked. */
 UCLASS()
@@ -19,36 +19,34 @@ class UARKitCameraOverlayMaterialLoader : public UObject
 public:
 	UPROPERTY()
 	UMaterialInterface* DefaultCameraOverlayMaterial;
-	
+
 	UARKitCameraOverlayMaterialLoader()
 	{
-		static ConstructorHelpers::FObjectFinder<UMaterialInterface> DefaultOverlayMaterialRef(TEXT("/AppleARKit/ARKitCameraMaterial.ARKitCameraMaterial"));
+		static ConstructorHelpers::FObjectFinder<UMaterialInterface> DefaultOverlayMaterialRef(TEXT("/AppleARKit/M_CameraOverlay.M_CameraOverlay"));
 		DefaultCameraOverlayMaterial = DefaultOverlayMaterialRef.Object;
 	}
 };
 
-struct FAppleARKitFrame;
-
 class FAppleARKitVideoOverlay
+	: public FGCObject
 {
 public:
 	FAppleARKitVideoOverlay();
 
-	void UpdateVideoTexture_RenderThread(FRHICommandListImmediate& RHICmdList, FAppleARKitFrame& Frame, const FSceneViewFamily& InViewFamily);
-	void RenderVideoOverlay_RenderThread(FRHICommandListImmediate& RHICmdList, const FSceneView& InView, const EDeviceScreenOrientation DeviceOrientation);
+	void SetCameraTexture(UARTextureCameraImage* InCameraImage);
+
+	void RenderVideoOverlay_RenderThread(FRHICommandListImmediate& RHICmdList, const FSceneView& InView);
 	bool GetPassthroughCameraUVs_RenderThread(TArray<FVector2D>& OutUVs, const EDeviceScreenOrientation DeviceOrientation);
 
-private:
-	FTextureRHIRef VideoTextureY;
-	FTextureRHIRef VideoTextureCbCr;
-	UMaterialInterface* RenderingOverlayMaterial;
-	FIndexBufferRHIRef OverlayIndexBufferRHI;
+	void SetOverlayTexture(UARTextureCameraImage* InCameraImage);
 
-	// Cache UV offset to be used by GetPassthroughCameraUVs_RenderThread
-	FVector2D UVOffset;
-	
-	// Separate vertex buffer for each supported device orientation
-	FVertexBufferRHIRef OverlayVertexBufferRHI[4];
+private:
+	//~ FGCObject
+	virtual void AddReferencedObjects(FReferenceCollector& Collector) override;
+	//~ FGCObject
 
-	double LastUpdateTimestamp;
+	UMaterialInstanceDynamic* MID_CameraOverlay;
+	UMaterialInterface* RenderingOverlayMaterial;
+	FIndexBufferRHIRef IndexBufferRHI;
+	FVertexBufferRHIRef VertexBufferRHI;
 };
