Changed the ARKit camera texture conversion to use linear color space

diff --git a/Engine/Plugins/Runtime/AR/Apple/AppleARKit/Source/AppleARKit/Private/AppleARKitTextures.cpp b/Engine/Plugins/Runtime/AR/Apple/AppleARKit/Source/AppleARKit/Private/AppleARKitTextures.cpp
--- a/Engine/Plugins/Runtime/AR/Apple/AppleARKit/Source/AppleARKit/Private/AppleARKitTextures.cpp
+++ b/Engine/Plugins/Runtime/AR/Apple/AppleARKit/Source/AppleARKit/Private/AppleARKitTextures.cpp
@@ -26,6 +26,7 @@ public:
 			CFRetain(CameraImage);
 		}
 #endif
+		bSRGB = false;
 	}
 
 	virtual ~FARKitCameraImageResource()
@@ -42,7 +43,7 @@ public:
 		{
 			SCOPED_AUTORELEASE_POOL;
 
-			CGColorSpaceRef ColorSpaceRef = CGColorSpaceCreateDeviceRGB();
+			CGColorSpaceRef ColorSpaceRef = CGColorSpaceCreateWithName(kCGColorSpaceGenericRGBLinear);
 			CIImage* Image = [[CIImage alloc] initWithCVPixelBuffer: CameraImage];
 
 			// Textures always need to be rotated so to a sane orientation (and mirrored because of differing coord system)
@@ -191,6 +192,7 @@ UAppleARKitTextureCameraImage::UAppleARKitTextureCameraImage(const FObjectInitia
 #endif
 {
 	ExternalTextureGuid = FGuid::NewGuid();
+	SRGB = false;
 }
 
 FTextureResource* UAppleARKitTextureCameraImage::CreateResource()
