Removed the unneeded casts on the vertex & pixel shaders for the ARKit overlay

diff --git a/Engine/Plugins/Runtime/AR/Apple/AppleARKit/Source/AppleARKit/Private/AppleARKitVideoOverlay.cpp b/Engine/Plugins/Runtime/AR/Apple/AppleARKit/Source/AppleARKit/Private/AppleARKitVideoOverlay.cpp
--- a/Engine/Plugins/Runtime/AR/Apple/AppleARKit/Source/AppleARKit/Private/AppleARKitVideoOverlay.cpp
+++ b/Engine/Plugins/Runtime/AR/Apple/AppleARKit/Source/AppleARKit/Private/AppleARKitVideoOverlay.cpp
@@ -223,15 +223,15 @@ void FAppleARKitVideoOverlay::RenderVideoOverlay_RenderThread(FRHICommandListImm
 	{
 		VertexShader = MaterialShaderMap->GetShader<FARKitCameraOverlayVS<true>>();
 		PixelShader = MaterialShaderMap->GetShader<FARKitCameraOverlayPS<true>>();
-		GraphicsPSOInit.BoundShaderState.VertexShaderRHI = GETSAFERHISHADER_VERTEX(reinterpret_cast<FARKitCameraOverlayVS<true>*>(VertexShader));
-		GraphicsPSOInit.BoundShaderState.PixelShaderRHI = GETSAFERHISHADER_PIXEL(reinterpret_cast<FARKitCameraOverlayPS<true>*>(PixelShader));
+		GraphicsPSOInit.BoundShaderState.VertexShaderRHI = GETSAFERHISHADER_VERTEX(VertexShader);
+		GraphicsPSOInit.BoundShaderState.PixelShaderRHI = GETSAFERHISHADER_PIXEL(PixelShader);
 	}
 	else
 	{
 		VertexShader = MaterialShaderMap->GetShader<FARKitCameraOverlayVS<false>>();
 		PixelShader = MaterialShaderMap->GetShader<FARKitCameraOverlayPS<false>>();
-		GraphicsPSOInit.BoundShaderState.VertexShaderRHI = GETSAFERHISHADER_VERTEX(reinterpret_cast<FARKitCameraOverlayVS<false>*>(VertexShader));
-		GraphicsPSOInit.BoundShaderState.PixelShaderRHI = GETSAFERHISHADER_PIXEL(reinterpret_cast<FARKitCameraOverlayPS<false>*>(PixelShader));
+		GraphicsPSOInit.BoundShaderState.VertexShaderRHI = GETSAFERHISHADER_VERTEX(VertexShader);
+		GraphicsPSOInit.BoundShaderState.PixelShaderRHI = GETSAFERHISHADER_PIXEL(PixelShader);
 	}
 
 	SetGraphicsPipelineState(RHICmdList, GraphicsPSOInit);
