More changes to casts in ARKit overlay

diff --git a/Engine/Plugins/Runtime/AR/Apple/AppleARKit/Source/AppleARKit/Private/AppleARKitVideoOverlay.cpp b/Engine/Plugins/Runtime/AR/Apple/AppleARKit/Source/AppleARKit/Private/AppleARKitVideoOverlay.cpp
--- a/Engine/Plugins/Runtime/AR/Apple/AppleARKit/Source/AppleARKit/Private/AppleARKitVideoOverlay.cpp
+++ b/Engine/Plugins/Runtime/AR/Apple/AppleARKit/Source/AppleARKit/Private/AppleARKitVideoOverlay.cpp
@@ -246,17 +246,17 @@ void FAppleARKitVideoOverlay::RenderVideoOverlay_RenderThread(FRHICommandListImm
 
 	if (bIsMobileRenderer)
 	{
-		FARKitCameraOverlayVS<true>* const VertexShaderPtr = reinterpret_cast<FARKitCameraOverlayVS<true>*>(VertexShader);
+		FARKitCameraOverlayVS<true>* const VertexShaderPtr = static_cast<FARKitCameraOverlayVS<true>*>(VertexShader);
 		SetUniformBufferParameterImmediate(RHICmdList, VertexShaderPtr->GetVertexShader(), VertexShaderPtr->GetUniformBufferParameter<FDrawRectangleParameters>(), Parameters);
 		VertexShaderPtr->SetParameters(RHICmdList, InView);
-		reinterpret_cast<FARKitCameraOverlayPS<true>*>(PixelShader)->SetParameters(RHICmdList, InView, RenderingOverlayMaterial->GetRenderProxy());
+		static_cast<FARKitCameraOverlayPS<true>*>(PixelShader)->SetParameters(RHICmdList, InView, RenderingOverlayMaterial->GetRenderProxy());
 	}
 	else
 	{
-		FARKitCameraOverlayVS<false>* const VertexShaderPtr = reinterpret_cast<FARKitCameraOverlayVS<false>*>(VertexShader);
+		FARKitCameraOverlayVS<false>* const VertexShaderPtr = static_cast<FARKitCameraOverlayVS<false>*>(VertexShader);
 		SetUniformBufferParameterImmediate(RHICmdList, VertexShaderPtr->GetVertexShader(), VertexShaderPtr->GetUniformBufferParameter<FDrawRectangleParameters>(), Parameters);
 		VertexShaderPtr->SetParameters(RHICmdList, InView);
-		reinterpret_cast<FARKitCameraOverlayPS<false>*>(PixelShader)->SetParameters(RHICmdList, InView, RenderingOverlayMaterial->GetRenderProxy());
+		static_cast<FARKitCameraOverlayPS<false>*>(PixelShader)->SetParameters(RHICmdList, InView, RenderingOverlayMaterial->GetRenderProxy());
 	}
 
 	RHICmdList.SetStreamSource(0, VertexBufferRHI, 0);
