Added some redundant check()s to make the static analysis realize they've been checked repeatedly

diff --git a/Engine/Plugins/Runtime/AR/Apple/AppleARKit/Source/AppleARKit/Private/AppleARKitVideoOverlay.cpp b/Engine/Plugins/Runtime/AR/Apple/AppleARKit/Source/AppleARKit/Private/AppleARKitVideoOverlay.cpp
--- a/Engine/Plugins/Runtime/AR/Apple/AppleARKit/Source/AppleARKit/Private/AppleARKitVideoOverlay.cpp
+++ b/Engine/Plugins/Runtime/AR/Apple/AppleARKit/Source/AppleARKit/Private/AppleARKitVideoOverlay.cpp
@@ -222,14 +222,18 @@ void FAppleARKitVideoOverlay::RenderVideoOverlay_RenderThread(FRHICommandListImm
 	if (bIsMobileRenderer)
 	{
 		VertexShader = MaterialShaderMap->GetShader<FARKitCameraOverlayVS<true>>();
+		check(VertexShader != nullptr);
 		PixelShader = MaterialShaderMap->GetShader<FARKitCameraOverlayPS<true>>();
+		check(PixelShader != nullptr);
 		GraphicsPSOInit.BoundShaderState.VertexShaderRHI = GETSAFERHISHADER_VERTEX(VertexShader);
 		GraphicsPSOInit.BoundShaderState.PixelShaderRHI = GETSAFERHISHADER_PIXEL(PixelShader);
 	}
 	else
 	{
 		VertexShader = MaterialShaderMap->GetShader<FARKitCameraOverlayVS<false>>();
+		check(VertexShader != nullptr);
 		PixelShader = MaterialShaderMap->GetShader<FARKitCameraOverlayPS<false>>();
+		check(PixelShader != nullptr);
 		GraphicsPSOInit.BoundShaderState.VertexShaderRHI = GETSAFERHISHADER_VERTEX(VertexShader);
 		GraphicsPSOInit.BoundShaderState.PixelShaderRHI = GETSAFERHISHADER_PIXEL(PixelShader);
 	}
@@ -247,16 +251,22 @@ void FAppleARKitVideoOverlay::RenderVideoOverlay_RenderThread(FRHICommandListImm
 	if (bIsMobileRenderer)
 	{
 		FARKitCameraOverlayVS<true>* const VertexShaderPtr = static_cast<FARKitCameraOverlayVS<true>*>(VertexShader);
+		check(VertexShaderPtr != nullptr);
 		SetUniformBufferParameterImmediate(RHICmdList, VertexShaderPtr->GetVertexShader(), VertexShaderPtr->GetUniformBufferParameter<FDrawRectangleParameters>(), Parameters);
 		VertexShaderPtr->SetParameters(RHICmdList, InView);
-		static_cast<FARKitCameraOverlayPS<true>*>(PixelShader)->SetParameters(RHICmdList, InView, RenderingOverlayMaterial->GetRenderProxy());
+		FARKitCameraOverlayPS<true>* PixelShaderPtr = static_cast<FARKitCameraOverlayPS<true>*>(PixelShader);
+		check(PixelShaderPtr != nullptr);
+		PixelShaderPtr->SetParameters(RHICmdList, InView, RenderingOverlayMaterial->GetRenderProxy());
 	}
 	else
 	{
 		FARKitCameraOverlayVS<false>* const VertexShaderPtr = static_cast<FARKitCameraOverlayVS<false>*>(VertexShader);
+		check(VertexShaderPtr != nullptr);
 		SetUniformBufferParameterImmediate(RHICmdList, VertexShaderPtr->GetVertexShader(), VertexShaderPtr->GetUniformBufferParameter<FDrawRectangleParameters>(), Parameters);
 		VertexShaderPtr->SetParameters(RHICmdList, InView);
-		static_cast<FARKitCameraOverlayPS<false>*>(PixelShader)->SetParameters(RHICmdList, InView, RenderingOverlayMaterial->GetRenderProxy());
+		FARKitCameraOverlayPS<false>* PixelShaderPtr = static_cast<FARKitCameraOverlayPS<false>*>(PixelShader);
+		check(PixelShaderPtr != nullptr);
+		PixelShaderPtr->SetParameters(RHICmdList, InView, RenderingOverlayMaterial->GetRenderProxy());
 	}
 
 	RHICmdList.SetStreamSource(0, VertexBufferRHI, 0);
