Deprecated the ARKit pass through material since it is no longer valid
- Use UARBlueprintLibrary::GetCameraImage() to get a UARTextureCameraImage object instead
- UVs are 0..1 like any normal texture 2d
- Texture is in linear space

diff --git a/Engine/Plugins/Runtime/AR/Apple/AppleARKit/Source/AppleARKit/Private/ExternalTextureMaterialExpression.cpp b/Engine/Plugins/Runtime/AR/Apple/AppleARKit/Source/AppleARKit/Private/ExternalTextureMaterialExpression.cpp
--- a/Engine/Plugins/Runtime/AR/Apple/AppleARKit/Source/AppleARKit/Private/ExternalTextureMaterialExpression.cpp
+++ b/Engine/Plugins/Runtime/AR/Apple/AppleARKit/Source/AppleARKit/Private/ExternalTextureMaterialExpression.cpp
@@ -4,12 +4,12 @@
 #include "MaterialCompiler.h"
 #include "ExternalTextureGuid.h"
 
-UMaterialExpressionARKitPassthroughCamera::UMaterialExpressionARKitPassthroughCamera(const FObjectInitializer& ObjectInitializer)
+UDEPRECATED_MaterialExpressionARKitPassthroughCamera::UDEPRECATED_MaterialExpressionARKitPassthroughCamera(const FObjectInitializer& ObjectInitializer)
 	: Super(ObjectInitializer)
 {}
 
 #if WITH_EDITOR
-int32 UMaterialExpressionARKitPassthroughCamera::Compile(class FMaterialCompiler* Compiler, int32 OutputIndex)
+int32 UDEPRECATED_MaterialExpressionARKitPassthroughCamera::Compile(class FMaterialCompiler* Compiler, int32 OutputIndex)
 {
 	return Compiler->TextureSample(
 		Compiler->ExternalTexture((TextureType == TextureY) ? ARKitPassthroughCameraExternalTextureYGuid : ARKitPassthroughCameraExternalTextureCbCrGuid),
@@ -17,13 +17,13 @@ int32 UMaterialExpressionARKitPassthroughCamera::Compile(class FMaterialCompiler
 		EMaterialSamplerType::SAMPLERTYPE_Color);
 }
 
-int32 UMaterialExpressionARKitPassthroughCamera::CompilePreview(class FMaterialCompiler* Compiler, int32 OutputIndex)
+int32 UDEPRECATED_MaterialExpressionARKitPassthroughCamera::CompilePreview(class FMaterialCompiler* Compiler, int32 OutputIndex)
 {
 	return INDEX_NONE;
 }
 
-void UMaterialExpressionARKitPassthroughCamera::GetCaption(TArray<FString>& OutCaptions) const
+void UDEPRECATED_MaterialExpressionARKitPassthroughCamera::GetCaption(TArray<FString>& OutCaptions) const
 {
-	OutCaptions.Add(TEXT("ARKit Passthrough Camera"));
+	OutCaptions.Add(TEXT("DEPRECATED USE ARSystem->GetCameraImage() - ARKit Passthrough Camera"));
 }
 #endif
\ No newline at end of file
diff --git a/Engine/Plugins/Runtime/AR/Apple/AppleARKit/Source/AppleARKit/Private/ExternalTextureMaterialExpression.h b/Engine/Plugins/Runtime/AR/Apple/AppleARKit/Source/AppleARKit/Private/ExternalTextureMaterialExpression.h
--- a/Engine/Plugins/Runtime/AR/Apple/AppleARKit/Source/AppleARKit/Private/ExternalTextureMaterialExpression.h
+++ b/Engine/Plugins/Runtime/AR/Apple/AppleARKit/Source/AppleARKit/Private/ExternalTextureMaterialExpression.h
@@ -17,8 +17,8 @@ enum EARKitTextureType
 /**
 * Implements a node sampling from the ARKit Passthrough external textures.
 */
-UCLASS(collapsecategories, hidecategories = Object)
-class UMaterialExpressionARKitPassthroughCamera : public UMaterialExpression
+UCLASS(collapsecategories, hidecategories = Object, Deprecated)
+class UDEPRECATED_MaterialExpressionARKitPassthroughCamera : public UMaterialExpression
 {
 	GENERATED_UCLASS_BODY()
 
