Removed ARKit classes that were marked deprecated and for removal in 4.22 or earlier

diff --git a/Engine/Plugins/Runtime/AR/Apple/AppleARKit/Source/AppleARKit/Private/AppleARKitAnchor.cpp b/Engine/Plugins/Runtime/AR/Apple/AppleARKit/Source/AppleARKit/Private/AppleARKitAnchor.cpp
deleted file mode 100644
diff --git a/Engine/Plugins/Runtime/AR/Apple/AppleARKit/Source/AppleARKit/Private/AppleARKitHitTestResult.cpp b/Engine/Plugins/Runtime/AR/Apple/AppleARKit/Source/AppleARKit/Private/AppleARKitHitTestResult.cpp
--- a/Engine/Plugins/Runtime/AR/Apple/AppleARKit/Source/AppleARKit/Private/AppleARKitHitTestResult.cpp
+++ b/Engine/Plugins/Runtime/AR/Apple/AppleARKit/Source/AppleARKit/Private/AppleARKitHitTestResult.cpp
@@ -62,16 +62,4 @@ ARHitTestResultType ToARHitTestResultType(EAppleARKitHitTestResultType InTypes)
 	return Types;
 }
 
-FAppleARKitHitTestResult::FAppleARKitHitTestResult( ARHitTestResult* InARHitTestResult, class UDEPRECATED_AppleARKitAnchor* InAnchor /*= nullptr*/, float WorldToMetersScale /*= 100.0f*/ )
-{
-	// Sanity check
-	check( InARHitTestResult );
-
-	// Convert properties
-	Type = ToEAppleARKitHitTestResultType( InARHitTestResult.type );
-    Distance = InARHitTestResult.distance * WorldToMetersScale;
-	Transform = FAppleARKitConversion::ToFTransform( InARHitTestResult.worldTransform );
-	Anchor_DEPRECATED = InAnchor;
-}
-
 #endif
diff --git a/Engine/Plugins/Runtime/AR/Apple/AppleARKit/Source/AppleARKit/Private/AppleARKitPlaneAnchor.cpp b/Engine/Plugins/Runtime/AR/Apple/AppleARKit/Source/AppleARKit/Private/AppleARKitPlaneAnchor.cpp
deleted file mode 100644
diff --git a/Engine/Plugins/Runtime/AR/Apple/AppleARKit/Source/AppleARKit/Private/AppleARKitSystem.cpp b/Engine/Plugins/Runtime/AR/Apple/AppleARKit/Source/AppleARKit/Private/AppleARKitSystem.cpp
--- a/Engine/Plugins/Runtime/AR/Apple/AppleARKit/Source/AppleARKit/Private/AppleARKitSystem.cpp
+++ b/Engine/Plugins/Runtime/AR/Apple/AppleARKit/Source/AppleARKit/Private/AppleARKitSystem.cpp
@@ -8,8 +8,6 @@
 #include "AppleARKitConversion.h"
 #include "AppleARKitVideoOverlay.h"
 #include "AppleARKitFrame.h"
-#include "AppleARKitAnchor.h"
-#include "AppleARKitPlaneAnchor.h"
 #include "AppleARKitConversion.h"
 #include "GeneralProjectSettings.h"
 #include "ARSessionConfig.h"
@@ -1810,10 +1808,6 @@ void FAppleARKitSystem::SessionDidUpdateAnchors_Internal( TSharedRef<FAppleARKit
 				{
 					FaceGeo->UpdateFaceGeometry(ARComponent.ToSharedRef(), AnchorData->FrameNumber, AnchorData->Timestamp, AnchorData->Transform, GetARCompositionComponent()->GetAlignmentTransform(), AnchorData->BlendShapes, AnchorData->FaceVerts, AnchorData->FaceIndices, TArray<FVector2D>(), AnchorData->LeftEyeTransform, AnchorData->RightEyeTransform, AnchorData->LookAtTarget);
 					FaceGeo->SetTrackingState(AnchorData->bIsTracked ? EARTrackingState::Tracking : EARTrackingState::NotTracking);
-					// @todo JoeG -- remove this in 4.22
-PRAGMA_DISABLE_DEPRECATION_WARNINGS
-					FaceGeo->bIsTracked = AnchorData->bIsTracked;
-PRAGMA_ENABLE_DEPRECATION_WARNINGS
 					for (UARPin* Pin : PinsToUpdate)
 					{
 						const FTransform Pin_LocalToTrackingTransform_PostUpdate = Pin->GetLocalToTrackingTransform_NoAlignment() * AnchorDeltaTransform;
@@ -1831,10 +1825,6 @@ PRAGMA_ENABLE_DEPRECATION_WARNINGS
 					FVector2D PhysicalSize((*CandidateImage)->GetPhysicalWidth(), (*CandidateImage)->GetPhysicalHeight());
 					ImageAnchor->UpdateTrackedGeometry(ARComponent.ToSharedRef(), AnchorData->FrameNumber, AnchorData->Timestamp, AnchorData->Transform, GetARCompositionComponent()->GetAlignmentTransform(), PhysicalSize, *CandidateImage);
 					ImageAnchor->SetTrackingState(AnchorData->bIsTracked ? EARTrackingState::Tracking : EARTrackingState::NotTracking);
-					// @todo JoeG -- remove this in 4.22
-PRAGMA_DISABLE_DEPRECATION_WARNINGS
-					ImageAnchor->bIsTracked = AnchorData->bIsTracked;
-PRAGMA_ENABLE_DEPRECATION_WARNINGS
 					for (UARPin* Pin : PinsToUpdate)
 					{
 						const FTransform Pin_LocalToTrackingTransform_PostUpdate = Pin->GetLocalToTrackingTransform_NoAlignment() * AnchorDeltaTransform;
diff --git a/Engine/Plugins/Runtime/AR/Apple/AppleARKit/Source/AppleARKit/Public/AppleARKitAnchor.h b/Engine/Plugins/Runtime/AR/Apple/AppleARKit/Source/AppleARKit/Public/AppleARKitAnchor.h
deleted file mode 100644
diff --git a/Engine/Plugins/Runtime/AR/Apple/AppleARKit/Source/AppleARKit/Public/AppleARKitHitTestResult.h b/Engine/Plugins/Runtime/AR/Apple/AppleARKit/Source/AppleARKit/Public/AppleARKitHitTestResult.h
--- a/Engine/Plugins/Runtime/AR/Apple/AppleARKit/Source/AppleARKit/Public/AppleARKitHitTestResult.h
+++ b/Engine/Plugins/Runtime/AR/Apple/AppleARKit/Source/AppleARKit/Public/AppleARKitHitTestResult.h
@@ -58,12 +58,6 @@ struct APPLEARKIT_API FAppleARKitHitTestResult
 
 #if SUPPORTS_ARKIT_1_0
 
-	/** 
-	 * This is a conversion copy-constructor that takes a raw ARHitTestResult and fills this 
-	 * structs members with the UE4-ified versions of ARHitTestResult's properties.
-	 */ 
-	FAppleARKitHitTestResult( ARHitTestResult* InARHitTestResult, class UDEPRECATED_AppleARKitAnchor* InAnchor = nullptr, float WorldToMetersScale = 100.0f );
-
 #endif
 
 	/**
@@ -84,12 +78,4 @@ struct APPLEARKIT_API FAppleARKitHitTestResult
 	 */
     UPROPERTY( BlueprintReadOnly, Category = "AppleARKitHitTestResult")
 	FTransform Transform;
-
-	/**
-	 * The anchor that the hit-test intersected.
-	 * 
-	 * An anchor will only be provided for existing plane result types.
-	 */
-	UPROPERTY()
-	class UDEPRECATED_AppleARKitAnchor* Anchor_DEPRECATED = nullptr;
 };
diff --git a/Engine/Plugins/Runtime/AR/Apple/AppleARKit/Source/AppleARKit/Public/AppleARKitPlaneAnchor.h b/Engine/Plugins/Runtime/AR/Apple/AppleARKit/Source/AppleARKit/Public/AppleARKitPlaneAnchor.h
deleted file mode 100644
diff --git a/Engine/Source/Runtime/AugmentedReality/Public/ARTrackable.h b/Engine/Source/Runtime/AugmentedReality/Public/ARTrackable.h
--- a/Engine/Source/Runtime/AugmentedReality/Public/ARTrackable.h
+++ b/Engine/Source/Runtime/AugmentedReality/Public/ARTrackable.h
@@ -189,11 +189,6 @@ public:
 	UFUNCTION(BlueprintPure, Category = "AR AugmentedReality|Image Detection")
 	FVector2D GetEstimateSize();
 
-	UE_DEPRECATED(4.21, "This property is now deprecated, please use GetTrackingState() and check for EARTrackingState::Tracking or IsTracked() instead.")
-	/** Whether the image is currently being tracked by the AR system */
-	UPROPERTY(BlueprintReadOnly, Category="AR AugmentedReality|Face Geometry")
-	bool bIsTracked;
-
 protected:
 	/** The candidate image that was detected in the scene */
 	UPROPERTY()
