Missed a few bIsTracked removal places

diff --git a/Engine/Plugins/Runtime/AR/Apple/AppleARKit/Source/AppleARKit/Private/AppleARKitSystem.cpp b/Engine/Plugins/Runtime/AR/Apple/AppleARKit/Source/AppleARKit/Private/AppleARKitSystem.cpp
--- a/Engine/Plugins/Runtime/AR/Apple/AppleARKit/Source/AppleARKit/Private/AppleARKitSystem.cpp
+++ b/Engine/Plugins/Runtime/AR/Apple/AppleARKit/Source/AppleARKit/Private/AppleARKitSystem.cpp
@@ -1663,10 +1663,6 @@ void FAppleARKitSystem::SessionDidAddAnchors_Internal( TSharedRef<FAppleARKitAnc
 			UARFaceGeometry* NewGeo = NewObject<UARFaceGeometry>();
 			NewGeo->UpdateFaceGeometry(ARComponent.ToSharedRef(), AnchorData->FrameNumber, AnchorData->Timestamp, AnchorData->Transform, GetARCompositionComponent()->GetAlignmentTransform(), AnchorData->BlendShapes, AnchorData->FaceVerts, AnchorData->FaceIndices, TArray<FVector2D>(), AnchorData->LeftEyeTransform, AnchorData->RightEyeTransform, AnchorData->LookAtTarget);
 			NewGeo->SetTrackingState(EARTrackingState::Tracking);
-			// @todo JoeG -- remove in 4.22
-PRAGMA_DISABLE_DEPRECATION_WARNINGS
-			NewGeo->bIsTracked = true;
-PRAGMA_ENABLE_DEPRECATION_WARNINGS
 			NewGeometry = NewGeo;
 			break;
 		}
diff --git a/Engine/Plugins/Runtime/AR/Google/GoogleARCore/Source/GoogleARCoreBase/Private/GoogleARCoreAugmentedImage.cpp b/Engine/Plugins/Runtime/AR/Google/GoogleARCore/Source/GoogleARCoreBase/Private/GoogleARCoreAugmentedImage.cpp
--- a/Engine/Plugins/Runtime/AR/Google/GoogleARCore/Source/GoogleARCoreBase/Private/GoogleARCoreAugmentedImage.cpp
+++ b/Engine/Plugins/Runtime/AR/Google/GoogleARCore/Source/GoogleARCoreBase/Private/GoogleARCoreAugmentedImage.cpp
@@ -78,9 +78,6 @@ void UGoogleARCoreAugmentedImage::UpdateTrackedGeometry(
 
 	ImageIndex = InImageIndex;
 	ImageName = InImageName;
-PRAGMA_DISABLE_DEPRECATION_WARNINGS
-	bIsTracked = (GetTrackingState() == EARTrackingState::Tracking);
-PRAGMA_ENABLE_DEPRECATION_WARNINGS
 }
 
 FVector UGoogleARCoreAugmentedImage::GetCenter() const
