Removed the unused texture cache and various Y/CrCb properties from ARKit classes

diff --git a/Engine/Plugins/Runtime/AR/Apple/AppleARKit/Source/AppleARKit/Private/AppleARKitFrame.cpp b/Engine/Plugins/Runtime/AR/Apple/AppleARKit/Source/AppleARKit/Private/AppleARKitFrame.cpp
--- a/Engine/Plugins/Runtime/AR/Apple/AppleARKit/Source/AppleARKit/Private/AppleARKitFrame.cpp
+++ b/Engine/Plugins/Runtime/AR/Apple/AppleARKit/Source/AppleARKit/Private/AppleARKitFrame.cpp
@@ -9,16 +9,10 @@
 FAppleARKitFrame::FAppleARKitFrame()
 	: Timestamp(0.0)
 #if SUPPORTS_ARKIT_1_0
-	, CapturedYImage(nullptr)
-	, CapturedCbCrImage(nullptr)
 	, CameraImage(nullptr)
 	, CameraDepth(nullptr)
 	, NativeFrame(nullptr)
 #endif
-	, CapturedYImageWidth(0)
-	, CapturedYImageHeight(0)
-	, CapturedCbCrImageWidth(0)
-	, CapturedCbCrImageHeight(0)
 	, WorldMappingState(EARWorldMappingState::NotAvailable)
 {
 };
@@ -42,14 +36,13 @@ EARWorldMappingState ToEARWorldMappingState(ARWorldMappingStatus MapStatus)
 
 #if SUPPORTS_ARKIT_1_0
 
-FAppleARKitFrame::FAppleARKitFrame( ARFrame* InARFrame, CVMetalTextureCacheRef MetalTextureCache )
+FAppleARKitFrame::FAppleARKitFrame( ARFrame* InARFrame )
 	: Camera( InARFrame.camera )
 	, LightEstimate( InARFrame.lightEstimate )
 	, WorldMappingState(EARWorldMappingState::NotAvailable)
 {
 	// Sanity check
 	check( InARFrame );
-	check( MetalTextureCache );
 
 	// Copy timestamp
 	Timestamp = InARFrame.timestamp;
@@ -57,42 +50,10 @@ FAppleARKitFrame::FAppleARKitFrame( ARFrame* InARFrame, CVMetalTextureCacheRef M
 	CameraImage = nullptr;
 	CameraDepth = nullptr;
 
-	// Copy / convert pass-through camera image's CVPixelBuffer to an MTLTexture so we can pass it
-	// directly to FTextureResource's.
-	// @see AppleARKitCameraTextureResource.cpp
-	CapturedYImage = nullptr;
-	CapturedYImageWidth = InARFrame.camera.imageResolution.width;
-	CapturedYImageHeight = InARFrame.camera.imageResolution.height;
-	
-	CapturedCbCrImage = nullptr;
-	CapturedCbCrImageWidth = InARFrame.camera.imageResolution.width;
-	CapturedCbCrImageHeight = InARFrame.camera.imageResolution.height;
-	
 	if ( InARFrame.capturedImage )
 	{
 		CameraImage = InARFrame.capturedImage;
 		CFRetain(CameraImage);
-
-		// Update SizeX & Y
-		CapturedYImageWidth = CVPixelBufferGetWidthOfPlane( InARFrame.capturedImage, 0 );
-		CapturedYImageHeight = CVPixelBufferGetHeightOfPlane( InARFrame.capturedImage, 0 );
-		CapturedCbCrImageWidth = CVPixelBufferGetWidthOfPlane( InARFrame.capturedImage, 1 );
-		CapturedCbCrImageHeight = CVPixelBufferGetHeightOfPlane( InARFrame.capturedImage, 1 );
-		
-		// Create a metal texture from the CVPixelBufferRef. The CVMetalTextureRef will
-		// be released in the FAppleARKitFrame destructor.
-		// NOTE: On success, CapturedImage will be a new CVMetalTextureRef with a ref count of 1
-		// 		 so we don't need to CFRetain it. The corresponding CFRelease is handled in
-		//
-		CVReturn Result = CVMetalTextureCacheCreateTextureFromImage( nullptr, MetalTextureCache, InARFrame.capturedImage, nullptr, MTLPixelFormatR8Unorm, CapturedYImageWidth, CapturedYImageHeight, /*PlaneIndex*/0, &CapturedYImage );
-		check( Result == kCVReturnSuccess );
-		check( CapturedYImage );
-		check( CFGetRetainCount(CapturedYImage) == 1);
-		
-		Result = CVMetalTextureCacheCreateTextureFromImage( nullptr, MetalTextureCache, InARFrame.capturedImage, nullptr, MTLPixelFormatRG8Unorm, CapturedCbCrImageWidth, CapturedCbCrImageHeight, /*PlaneIndex*/1, &CapturedCbCrImage );
-		check( Result == kCVReturnSuccess );
-		check( CapturedCbCrImage );
-		check( CFGetRetainCount(CapturedCbCrImage) == 1);
 	}
 
 	if (InARFrame.capturedDepthData)
@@ -113,14 +74,8 @@ FAppleARKitFrame::FAppleARKitFrame( ARFrame* InARFrame, CVMetalTextureCacheRef M
 
 FAppleARKitFrame::FAppleARKitFrame( const FAppleARKitFrame& Other )
 	: Timestamp( Other.Timestamp )
-	, CapturedYImage( nullptr )
-	, CapturedCbCrImage( nullptr )
 	, CameraImage( nullptr )
 	, CameraDepth( nullptr )
-	, CapturedYImageWidth( Other.CapturedYImageWidth )
-	, CapturedYImageHeight( Other.CapturedYImageHeight )
-	, CapturedCbCrImageWidth( Other.CapturedCbCrImageWidth )
-	, CapturedCbCrImageHeight( Other.CapturedCbCrImageHeight )
 	, Camera( Other.Camera )
 	, LightEstimate( Other.LightEstimate )
 	, WorldMappingState(Other.WorldMappingState)
@@ -134,14 +89,6 @@ FAppleARKitFrame::FAppleARKitFrame( const FAppleARKitFrame& Other )
 FAppleARKitFrame::~FAppleARKitFrame()
 {
 	// Release captured image
-	if ( CapturedYImage != nullptr )
-	{
-		CFRelease( CapturedYImage );
-	}
-	if ( CapturedCbCrImage != nullptr )
-	{
-		CFRelease( CapturedCbCrImage );
-	}
 	if (CameraImage != nullptr)
 	{
 		CFRelease(CameraImage);
@@ -164,14 +111,6 @@ FAppleARKitFrame& FAppleARKitFrame::operator=( const FAppleARKitFrame& Other )
 	}
 
 	// Release outgoing image
-	if ( CapturedYImage != nullptr )
-	{
-		CFRelease( CapturedYImage );
-	}
-	if ( CapturedCbCrImage != nullptr )
-	{
-		CFRelease( CapturedCbCrImage );
-	}
 	if (CameraImage != nullptr)
 	{
 		CFRelease(CameraImage);
@@ -195,12 +134,6 @@ FAppleARKitFrame& FAppleARKitFrame::operator=( const FAppleARKitFrame& Other )
 	
 	// Member-wise copy
 	Timestamp = Other.Timestamp;
-	CapturedYImage = nullptr;
-	CapturedYImageWidth = Other.CapturedYImageWidth;
-	CapturedYImageHeight = Other.CapturedYImageHeight;
-	CapturedCbCrImage = nullptr;
-	CapturedCbCrImageWidth = Other.CapturedCbCrImageWidth;
-	CapturedCbCrImageHeight = Other.CapturedCbCrImageHeight;
 	Camera = Other.Camera;
 	LightEstimate = Other.LightEstimate;
 	WorldMappingState = Other.WorldMappingState;
diff --git a/Engine/Plugins/Runtime/AR/Apple/AppleARKit/Source/AppleARKit/Private/AppleARKitSessionDelegate.cpp b/Engine/Plugins/Runtime/AR/Apple/AppleARKit/Source/AppleARKit/Private/AppleARKitSessionDelegate.cpp
--- a/Engine/Plugins/Runtime/AR/Apple/AppleARKit/Source/AppleARKit/Private/AppleARKitSessionDelegate.cpp
+++ b/Engine/Plugins/Runtime/AR/Apple/AppleARKit/Source/AppleARKit/Private/AppleARKitSessionDelegate.cpp
@@ -11,7 +11,6 @@
 @implementation FAppleARKitSessionDelegate
 {
 	FAppleARKitSystem* _AppleARKitSystem;
-	CVMetalTextureCacheRef _metalTextureCache;
 }
 
 
@@ -23,38 +22,16 @@
 	{
 		UE_LOG(LogAppleARKit, Log, TEXT("Delegate created with session: %p"), InAppleARKitSystem);
 		_AppleARKitSystem = InAppleARKitSystem;
-		_metalTextureCache = NULL;	
 	}
 	return self;
 }
 
-- (void)setMetalTextureCache:(CVMetalTextureCacheRef)InMetalTextureCache
-{
-	// Release current?
-	if ( _metalTextureCache != nullptr )
-	{
-		CFRelease( _metalTextureCache );
-	}
-	// Set new & retain
-	_metalTextureCache = InMetalTextureCache;
-	if ( _metalTextureCache != nullptr )
-	{
-		CFRetain( _metalTextureCache );
-	}
-}
 #pragma mark - ARSessionDelegate Methods
 
 - (void)session:(ARSession *)session didUpdateFrame:(ARFrame *)frame 
 {
-	// check if we get an update when the session is not initialized.
-	if (!_metalTextureCache)
-	{
-		UE_LOG(LogAppleARKit, Log, TEXT("Delegate didUpdateFrame with no valid _metalTextureCache (ignoring)"));
-		return;
-	}
-
 	// Bundle results into FAppleARKitFrame
-	TSharedPtr< FAppleARKitFrame, ESPMode::ThreadSafe > AppleARKitFrame( new FAppleARKitFrame( frame, _metalTextureCache ) );
+	TSharedPtr< FAppleARKitFrame, ESPMode::ThreadSafe > AppleARKitFrame( new FAppleARKitFrame( frame ) );
 	
 	// Pass result to session
 	_AppleARKitSystem->SessionDidUpdateFrame_DelegateThread( AppleARKitFrame );
diff --git a/Engine/Plugins/Runtime/AR/Apple/AppleARKit/Source/AppleARKit/Private/AppleARKitSessionDelegate.h b/Engine/Plugins/Runtime/AR/Apple/AppleARKit/Source/AppleARKit/Private/AppleARKitSessionDelegate.h
--- a/Engine/Plugins/Runtime/AR/Apple/AppleARKit/Source/AppleARKit/Private/AppleARKitSessionDelegate.h
+++ b/Engine/Plugins/Runtime/AR/Apple/AppleARKit/Source/AppleARKit/Private/AppleARKitSessionDelegate.h
@@ -16,8 +16,6 @@
 
 - (id)initWithAppleARKitSystem:(class FAppleARKitSystem*)InAppleARKitSystem;
 
-- (void)setMetalTextureCache:(CVMetalTextureCacheRef)InMetalTextureCache;
-
 @end
 
 #endif
diff --git a/Engine/Plugins/Runtime/AR/Apple/AppleARKit/Source/AppleARKit/Private/AppleARKitSystem.cpp b/Engine/Plugins/Runtime/AR/Apple/AppleARKit/Source/AppleARKit/Private/AppleARKitSystem.cpp
--- a/Engine/Plugins/Runtime/AR/Apple/AppleARKit/Source/AppleARKit/Private/AppleARKitSystem.cpp
+++ b/Engine/Plugins/Runtime/AR/Apple/AppleARKit/Source/AppleARKit/Private/AppleARKitSystem.cpp
@@ -1249,20 +1249,6 @@ bool FAppleARKitSystem::Run(UARSessionConfig* SessionConfig)
 				ClearTrackedGeometries();
 			}
 		}
-
-		// Create MetalTextureCache
-		if (IsMetalPlatform(GMaxRHIShaderPlatform))
-		{
-			id<MTLDevice> Device = (id<MTLDevice>)GDynamicRHI->RHIGetNativeDevice();
-			check(Device);
-
-			CVReturn Return = CVMetalTextureCacheCreate(nullptr, nullptr, Device, nullptr, &MetalTextureCache);
-			check(Return == kCVReturnSuccess);
-			check(MetalTextureCache);
-
-			// Pass to session delegate to use for Metal texture creation
-			[Delegate setMetalTextureCache : MetalTextureCache];
-		}
 		
 #if PLATFORM_IOS && !PLATFORM_TVOS
 		// Check if we need to adjust the priorities to allow ARKit to have more CPU time
@@ -1319,16 +1305,6 @@ bool FAppleARKitSystem::Pause()
 	{
 		// Suspend the session
 		[Session pause];
-	
-		// Release MetalTextureCache created in Start
-		if (MetalTextureCache)
-		{
-			// Tell delegate to release it
-			[Delegate setMetalTextureCache:nullptr];
-		
-			CFRelease(MetalTextureCache);
-			MetalTextureCache = nullptr;
-		}
 	}
 	
 #if PLATFORM_IOS && !PLATFORM_TVOS
diff --git a/Engine/Plugins/Runtime/AR/Apple/AppleARKit/Source/AppleARKit/Private/AppleARKitSystem.h b/Engine/Plugins/Runtime/AR/Apple/AppleARKit/Source/AppleARKit/Private/AppleARKitSystem.h
--- a/Engine/Plugins/Runtime/AR/Apple/AppleARKit/Source/AppleARKit/Private/AppleARKitSystem.h
+++ b/Engine/Plugins/Runtime/AR/Apple/AppleARKit/Source/AppleARKit/Private/AppleARKitSystem.h
@@ -163,9 +163,6 @@ private:
 	
 	// ARKit Session Delegate
 	FAppleARKitSessionDelegate* Delegate = nullptr;
-	
-	/** The Metal texture cache for unbuffered texture uploads. */
-	CVMetalTextureCacheRef MetalTextureCache = nullptr;
 
 	/** Cache of images that we've converted previously to prevent repeated conversion */
 	TMap< FString, CGImage* > ConvertedCandidateImages;
diff --git a/Engine/Plugins/Runtime/AR/Apple/AppleARKit/Source/AppleARKit/Public/AppleARKitFrame.h b/Engine/Plugins/Runtime/AR/Apple/AppleARKit/Source/AppleARKit/Public/AppleARKitFrame.h
--- a/Engine/Plugins/Runtime/AR/Apple/AppleARKit/Source/AppleARKit/Public/AppleARKitFrame.h
+++ b/Engine/Plugins/Runtime/AR/Apple/AppleARKit/Source/AppleARKit/Public/AppleARKitFrame.h
@@ -28,9 +28,8 @@ struct APPLEARKIT_API FAppleARKitFrame
 	 * with the UE4-ified versions of ARFrames properties.
 	 *
 	 * @param InARFrame 		- The frame to convert / copy
-	 * @param MetalTextureCache - A CVMetalTextureCacheRef to use in the conversion of InARFrame's CVPixelBufferRef to CVMetalTextureRef for our CapturedImage property.
-	 */ 
-	FAppleARKitFrame( ARFrame* InARFrame, CVMetalTextureCacheRef MetalTextureCache );
+	 */
+	FAppleARKitFrame( ARFrame* InARFrame );
 
 	/** 
 	 * Copy contructor. CapturedImage is skipped as we don't need / want to retain access to the 
@@ -53,10 +52,6 @@ struct APPLEARKIT_API FAppleARKitFrame
 
 #if SUPPORTS_ARKIT_1_0
 
-	/** The frame's captured images */
-    CVMetalTextureRef CapturedYImage;
-    CVMetalTextureRef CapturedCbCrImage;
-
 	/** The raw camera buffer from ARKit */
 	CVPixelBufferRef CameraImage;
 	/** The raw camera depth info from ARKit (needs iPhone X) */
@@ -64,13 +59,6 @@ struct APPLEARKIT_API FAppleARKitFrame
 	void* NativeFrame;
 #endif
 
-	/** The width and height in pixels of the frame's captured images. */
-	uint32 CapturedYImageWidth;
-	uint32 CapturedYImageHeight;
-    
-    uint32 CapturedCbCrImageWidth;
-    uint32 CapturedCbCrImageHeight;
-    
 	/** The camera used to capture the frame's image. */
 	FAppleARKitCamera Camera;
 
