Removed some code that was resurrected during the HoloLens merge but shouldn't have been

diff --git a/Engine/Plugins/Runtime/AR/Apple/AppleARKit/Source/AppleARKit/Private/AppleARKitSystem.cpp b/Engine/Plugins/Runtime/AR/Apple/AppleARKit/Source/AppleARKit/Private/AppleARKitSystem.cpp
--- a/Engine/Plugins/Runtime/AR/Apple/AppleARKit/Source/AppleARKit/Private/AppleARKitSystem.cpp
+++ b/Engine/Plugins/Runtime/AR/Apple/AppleARKit/Source/AppleARKit/Private/AppleARKitSystem.cpp
@@ -1195,11 +1195,7 @@ bool FAppleARKitSystem::Run(UARSessionConfig* SessionConfig)
 #if SUPPORTS_ARKIT_1_0
 	// Don't do the conversion work if they don't want this
 	FAppleARKitAnchorData::bGenerateGeometry = SessionConfig->bGenerateMeshDataFromTrackedGeometry;
-	// Set this based upon the project settings
-	bShouldWriteCameraImagePerFrame = GetDefault<UAppleARKitSettings>()->bShouldWriteCameraImagePerFrame;
-	WrittenCameraImageScale = GetDefault<UAppleARKitSettings>()->WrittenCameraImageScale;
-	WrittenCameraImageRotation = GetDefault<UAppleARKitSettings>()->WrittenCameraImageRotation;
-	WrittenCameraImageQuality = GetDefault<UAppleARKitSettings>()->WrittenCameraImageQuality;
+
 	if (FAppleARKitAvailability::SupportsARKit10())
 	{
 		ARSessionRunOptions options = 0;
