Fixed the ARKit camera being allocated on the render thread causing assertions

diff --git a/Engine/Plugins/Runtime/AR/Apple/AppleARKit/Source/AppleARKit/Private/AppleARKitSystem.cpp b/Engine/Plugins/Runtime/AR/Apple/AppleARKit/Source/AppleARKit/Private/AppleARKitSystem.cpp
--- a/Engine/Plugins/Runtime/AR/Apple/AppleARKit/Source/AppleARKit/Private/AppleARKitSystem.cpp
+++ b/Engine/Plugins/Runtime/AR/Apple/AppleARKit/Source/AppleARKit/Private/AppleARKitSystem.cpp
@@ -409,7 +409,8 @@ bool FAppleARKitSystem::IsHeadTrackingAllowed() const
 
 TSharedPtr<class IXRCamera, ESPMode::ThreadSafe> FAppleARKitSystem::GetXRCamera(int32 DeviceId)
 {
-	if (!XRCamera.IsValid())
+	// Don't create/load UObjects on the render thread
+	if (!XRCamera.IsValid() && IsInGameThread())
 	{
 		TSharedRef<FAppleARKitXRCamera, ESPMode::ThreadSafe> NewCamera = FSceneViewExtensions::NewExtension<FAppleARKitXRCamera>(*this, DeviceId);
 		XRCamera = NewCamera;
