Fixed face ar config being applied if the face ar plugin is enabled in your project but you want a non-face tracking ar session

diff --git a/Engine/Plugins/Runtime/AR/Apple/AppleARKit/Source/AppleARKit/Private/AppleARKitSystem.cpp b/Engine/Plugins/Runtime/AR/Apple/AppleARKit/Source/AppleARKit/Private/AppleARKitSystem.cpp
--- a/Engine/Plugins/Runtime/AR/Apple/AppleARKit/Source/AppleARKit/Private/AppleARKitSystem.cpp
+++ b/Engine/Plugins/Runtime/AR/Apple/AppleARKit/Source/AppleARKit/Private/AppleARKitSystem.cpp
@@ -1216,13 +1216,13 @@ bool FAppleARKitSystem::Run(UARSessionConfig* SessionConfig)
 
 		ARConfiguration* Configuration = nullptr;
 		CheckForFaceARSupport(SessionConfig);
-		if (FaceARSupport == nullptr)
+		if (FaceARSupport != nullptr && SessionConfig->GetSessionType() == EARSessionType::Face)
 		{
-			Configuration = FAppleARKitConversion::ToARConfiguration(SessionConfig, CandidateImages, ConvertedCandidateImages, CandidateObjects);
+			Configuration = FaceARSupport->ToARConfiguration(SessionConfig, TimecodeProvider);
 		}
 		else
 		{
-			Configuration = FaceARSupport->ToARConfiguration(SessionConfig, TimecodeProvider);
+			Configuration = FAppleARKitConversion::ToARConfiguration(SessionConfig, CandidateImages, ConvertedCandidateImages, CandidateObjects);
 		}
 
 		// Not all session types are supported by all devices
