Fixed the case where the cached rotation value was incorrect when switching between AR sessions using different tracking types

diff --git a/Engine/Plugins/Runtime/AR/Apple/AppleARKit/Source/AppleARKit/Private/AppleARKitSystem.cpp b/Engine/Plugins/Runtime/AR/Apple/AppleARKit/Source/AppleARKit/Private/AppleARKitSystem.cpp
--- a/Engine/Plugins/Runtime/AR/Apple/AppleARKit/Source/AppleARKit/Private/AppleARKitSystem.cpp
+++ b/Engine/Plugins/Runtime/AR/Apple/AppleARKit/Source/AppleARKit/Private/AppleARKitSystem.cpp
@@ -1148,11 +1148,8 @@ void FAppleARKitSystem::SetDeviceOrientation(EDeviceScreenOrientation InOrientat
 		InOrientation = EDeviceScreenOrientation::LandscapeLeft;
 	}
 
-	if (DeviceOrientation != InOrientation)
-	{
-		DeviceOrientation = InOrientation;
-		CalcTrackingToWorldRotation();
-	}
+	DeviceOrientation = InOrientation;
+	CalcTrackingToWorldRotation();
 }
 
 void FAppleARKitSystem::ClearTrackedGeometries()
@@ -1198,14 +1195,6 @@ bool FAppleARKitSystem::Run(UARSessionConfig* SessionConfig)
 		LastReceivedFrame = TSharedPtr<FAppleARKitFrame, ESPMode::ThreadSafe>();
 	}
 
-	// Make sure this is set at session start, because there are timing issues with using only the delegate approach
-	if (DeviceOrientation == EDeviceScreenOrientation::Unknown)
-	{
-		EDeviceScreenOrientation ScreenOrientation = FPlatformMisc::GetDeviceOrientation();
-		SetDeviceOrientation( ScreenOrientation );
-	}
-
-
 #if SUPPORTS_ARKIT_1_0
 	// Don't do the conversion work if they don't want this
 	FAppleARKitAnchorData::bGenerateGeometry = SessionConfig->bGenerateMeshDataFromTrackedGeometry;
@@ -1284,7 +1273,12 @@ bool FAppleARKitSystem::Run(UARSessionConfig* SessionConfig)
 	}
 	
 #endif
-	
+
+	// Make sure this is set at session start, because there are timing issues with using only the delegate approach
+	// Also this needs to be set each time a new session is started in case we switch tracking modes (gravity vs face)
+	EDeviceScreenOrientation ScreenOrientation = FPlatformMisc::GetDeviceOrientation();
+	SetDeviceOrientation( ScreenOrientation );
+
 	// @todo arkit Add support for relocating ARKit space to Unreal World Origin? BaseTransform = FTransform::Identity;
 	
 	// Set running state
