Fixing build break

diff --git a/Engine/Plugins/Runtime/AR/Apple/AppleARKit/Source/AppleARKit/Private/AppleARKitVideoOverlay.cpp b/Engine/Plugins/Runtime/AR/Apple/AppleARKit/Source/AppleARKit/Private/AppleARKitVideoOverlay.cpp
--- a/Engine/Plugins/Runtime/AR/Apple/AppleARKit/Source/AppleARKit/Private/AppleARKitVideoOverlay.cpp
+++ b/Engine/Plugins/Runtime/AR/Apple/AppleARKit/Source/AppleARKit/Private/AppleARKitVideoOverlay.cpp
@@ -220,21 +220,21 @@ void FAppleARKitVideoOverlay::RenderVideoOverlay_RenderThread(FRHICommandListImm
 
 	if (bIsMobileRenderer)
 	{
-		FARKitCameraOverlayVS<true>* const VertexShaderPtr = static_cast<FARKitCameraOverlayMobileVS*>(VertexShader);
+		FARKitCameraOverlayMobileVS* const VertexShaderPtr = static_cast<FARKitCameraOverlayMobileVS*>(VertexShader);
 		check(VertexShaderPtr != nullptr);
 		SetUniformBufferParameterImmediate(RHICmdList, VertexShaderPtr->GetVertexShader(), VertexShaderPtr->GetUniformBufferParameter<FDrawRectangleParameters>(), Parameters);
 		VertexShaderPtr->SetParameters(RHICmdList, InView);
-		FARKitCameraOverlayPS<true>* PixelShaderPtr = static_cast<FARKitCameraOverlayMobilePS*>(PixelShader);
+		FARKitCameraOverlayMobilePS* PixelShaderPtr = static_cast<FARKitCameraOverlayMobilePS*>(PixelShader);
 		check(PixelShaderPtr != nullptr);
 		PixelShaderPtr->SetParameters(RHICmdList, InView, RenderingOverlayMaterial->GetRenderProxy());
 	}
 	else
 	{
-		FARKitCameraOverlayVS<false>* const VertexShaderPtr = static_cast<FARKitCameraOverlayVS*>(VertexShader);
+		FARKitCameraOverlayVS* const VertexShaderPtr = static_cast<FARKitCameraOverlayVS*>(VertexShader);
 		check(VertexShaderPtr != nullptr);
 		SetUniformBufferParameterImmediate(RHICmdList, VertexShaderPtr->GetVertexShader(), VertexShaderPtr->GetUniformBufferParameter<FDrawRectangleParameters>(), Parameters);
 		VertexShaderPtr->SetParameters(RHICmdList, InView);
-		FARKitCameraOverlayPS<false>* PixelShaderPtr = static_cast<FARKitCameraOverlayPS*>(PixelShader);
+		FARKitCameraOverlayPS* PixelShaderPtr = static_cast<FARKitCameraOverlayPS*>(PixelShader);
 		check(PixelShaderPtr != nullptr);
 		PixelShaderPtr->SetParameters(RHICmdList, InView, RenderingOverlayMaterial->GetRenderProxy());
 	}
