Fixed build break on iOS.

diff --git a/Engine/Plugins/Runtime/AR/Apple/AppleARKit/Source/AppleARKit/Private/AppleARKitVideoOverlay.cpp b/Engine/Plugins/Runtime/AR/Apple/AppleARKit/Source/AppleARKit/Private/AppleARKitVideoOverlay.cpp
--- a/Engine/Plugins/Runtime/AR/Apple/AppleARKit/Source/AppleARKit/Private/AppleARKitVideoOverlay.cpp
+++ b/Engine/Plugins/Runtime/AR/Apple/AppleARKit/Source/AppleARKit/Private/AppleARKitVideoOverlay.cpp
@@ -72,6 +72,7 @@ template <bool bIsMobileRenderer>
 class TARKitCameraOverlayVS : public TPostProcessMaterialShader<bIsMobileRenderer>
 {
 public:
+	using FMaterialShader::FPermutationDomain;
 	DECLARE_MATERIAL_SHADER(TARKitCameraOverlayVS);
 
 	static void ModifyCompilationEnvironment(const FMaterialShaderPermutationParameters& Parameters, FShaderCompilerEnvironment& OutEnvironment)
@@ -102,6 +103,7 @@ template <bool bIsMobileRenderer>
 class TARKitCameraOverlayPS : public TPostProcessMaterialShader<bIsMobileRenderer>
 {
 public:
+	using FMaterialShader::FPermutationDomain;
 	DECLARE_MATERIAL_SHADER(TARKitCameraOverlayPS);
 
 	static void ModifyCompilationEnvironment(const FMaterialShaderPermutationParameters& Parameters, FShaderCompilerEnvironment& OutEnvironment)
