Mac build fixes.

diff --git a/Engine/Plugins/Runtime/AR/Apple/AppleARKit/Source/AppleARKit/Private/AppleARKitVideoOverlay.cpp b/Engine/Plugins/Runtime/AR/Apple/AppleARKit/Source/AppleARKit/Private/AppleARKitVideoOverlay.cpp
--- a/Engine/Plugins/Runtime/AR/Apple/AppleARKit/Source/AppleARKit/Private/AppleARKitVideoOverlay.cpp
+++ b/Engine/Plugins/Runtime/AR/Apple/AppleARKit/Source/AppleARKit/Private/AppleARKitVideoOverlay.cpp
@@ -77,27 +77,27 @@ public:
 
 	static void ModifyCompilationEnvironment(const FMaterialShaderPermutationParameters& Parameters, FShaderCompilerEnvironment& OutEnvironment)
 	{
-		TPostProcessMaterialShader::ModifyCompilationEnvironment(Parameters, OutEnvironment);
+		TPostProcessMaterialShader<bIsMobileRenderer>::ModifyCompilationEnvironment(Parameters, OutEnvironment);
 		OutEnvironment.SetDefine(TEXT("POST_PROCESS_AR_PASSTHROUGH"), 1);
 	}
 
 	TARKitCameraOverlayVS() = default;
 	TARKitCameraOverlayVS(const ShaderMetaType::CompiledShaderInitializerType& Initializer)
-		: TPostProcessMaterialShader(Initializer)
+		: TPostProcessMaterialShader<bIsMobileRenderer>(Initializer)
 	{}
 
 	void SetParameters(FRHICommandList& RHICmdList, const FSceneView View)
 	{
-		FRHIVertexShader* ShaderRHI = GetVertexShader();
-		TPostProcessMaterialShader::SetViewParameters(RHICmdList, ShaderRHI, View, View.ViewUniformBuffer);
+		FRHIVertexShader* ShaderRHI = this->GetVertexShader();
+		TPostProcessMaterialShader<bIsMobileRenderer>::SetViewParameters(RHICmdList, ShaderRHI, View, View.ViewUniformBuffer);
 	}
 };
 
 using FARKitCameraOverlayVS = TARKitCameraOverlayVS<false>;
 using FARKitCameraOverlayMobileVS = TARKitCameraOverlayVS<true>;
 
-IMPLEMENT_MATERIAL_SHADER(FARKitCameraOverlayVS, "/Engine/Private/PostProcessMaterialShaders.usf", "MainVS_VideoOverlay", SF_Vertex);
-IMPLEMENT_MATERIAL_SHADER(FARKitCameraOverlayMobileVS, "/Engine/Private/PostProcessMaterialShaders.usf", "MainVS", SF_Vertex);
+template<> IMPLEMENT_MATERIAL_SHADER(FARKitCameraOverlayVS, "/Engine/Private/PostProcessMaterialShaders.usf", "MainVS_VideoOverlay", SF_Vertex);
+template<> IMPLEMENT_MATERIAL_SHADER(FARKitCameraOverlayMobileVS, "/Engine/Private/PostProcessMaterialShaders.usf", "MainVS", SF_Vertex);
 
 template <bool bIsMobileRenderer>
 class TARKitCameraOverlayPS : public TPostProcessMaterialShader<bIsMobileRenderer>
@@ -108,27 +108,27 @@ public:
 
 	static void ModifyCompilationEnvironment(const FMaterialShaderPermutationParameters& Parameters, FShaderCompilerEnvironment& OutEnvironment)
 	{
-		TPostProcessMaterialShader::ModifyCompilationEnvironment(Parameters, OutEnvironment);
+		TPostProcessMaterialShader<bIsMobileRenderer>::ModifyCompilationEnvironment(Parameters, OutEnvironment);
 		OutEnvironment.SetDefine(TEXT("OUTPUT_MOBILE_HDR"), IsMobileHDR() ? 1 : 0);
 	}
 
 	TARKitCameraOverlayPS() = default;
 	TARKitCameraOverlayPS(const ShaderMetaType::CompiledShaderInitializerType& Initializer)
-		: TPostProcessMaterialShader(Initializer)
+		: TPostProcessMaterialShader<bIsMobileRenderer>(Initializer)
 	{}
 
 	void SetParameters(FRHICommandList& RHICmdList, const FSceneView View, const FMaterialRenderProxy* Material)
 	{
-		FRHIPixelShader* ShaderRHI = GetPixelShader();
-		TPostProcessMaterialShader::SetParameters(RHICmdList, ShaderRHI, Material, *Material->GetMaterial(View.GetFeatureLevel()), View, View.ViewUniformBuffer, ESceneTextureSetupMode::None);
+		FRHIPixelShader* ShaderRHI = this->GetPixelShader();
+		TPostProcessMaterialShader<bIsMobileRenderer>::SetParameters(RHICmdList, ShaderRHI, Material, *Material->GetMaterial(View.GetFeatureLevel()), View, View.ViewUniformBuffer, ESceneTextureSetupMode::None);
 	}
 };
 
 using FARKitCameraOverlayPS = TARKitCameraOverlayPS<false>;
 using FARKitCameraOverlayMobilePS = TARKitCameraOverlayPS<true>;
 
-IMPLEMENT_MATERIAL_SHADER(FARKitCameraOverlayPS, "/Engine/Private/PostProcessMaterialShaders.usf", "MainPS_VideoOverlay", SF_Pixel);
-IMPLEMENT_MATERIAL_SHADER(FARKitCameraOverlayMobilePS, "/Engine/Private/PostProcessMaterialShaders.usf", "MainPS_ES2", SF_Pixel);
+template<> IMPLEMENT_MATERIAL_SHADER(FARKitCameraOverlayPS, "/Engine/Private/PostProcessMaterialShaders.usf", "MainPS_VideoOverlay", SF_Pixel);
+template<> IMPLEMENT_MATERIAL_SHADER(FARKitCameraOverlayMobilePS, "/Engine/Private/PostProcessMaterialShaders.usf", "MainPS_ES2", SF_Pixel);
 
 void FAppleARKitVideoOverlay::RenderVideoOverlay_RenderThread(FRHICommandListImmediate& RHICmdList, const FSceneView& InView)
 {
