Fix ARKit build break (SetDeviceOrientation was renamed, one reference wasn't fixed).

diff --git a/Engine/Plugins/Runtime/AR/Apple/AppleARKit/Source/AppleARKit/Private/AppleARKitSystem.cpp b/Engine/Plugins/Runtime/AR/Apple/AppleARKit/Source/AppleARKit/Private/AppleARKitSystem.cpp
--- a/Engine/Plugins/Runtime/AR/Apple/AppleARKit/Source/AppleARKit/Private/AppleARKitSystem.cpp
+++ b/Engine/Plugins/Runtime/AR/Apple/AppleARKit/Source/AppleARKit/Private/AppleARKitSystem.cpp
@@ -1390,7 +1390,7 @@ bool FAppleARKitSystem::Run(UARSessionConfig* SessionConfig)
 	// Make sure this is set at session start, because there are timing issues with using only the delegate approach
 	// Also this needs to be set each time a new session is started in case we switch tracking modes (gravity vs face)
 	EDeviceScreenOrientation ScreenOrientation = FPlatformMisc::GetDeviceOrientation();
-	SetDeviceOrientation( ScreenOrientation );
+	SetDeviceOrientationAndDerivedTracking(ScreenOrientation);
 
 	// @todo arkit Add support for relocating ARKit space to Unreal World Origin? BaseTransform = FTransform::Identity;
 	
