Integration from ARKit v3 stream.

diff --git a/Engine/Plugins/Runtime/AR/Apple/AppleARKit/Source/AppleARKit/Private/AppleARKitConversion.cpp b/Engine/Plugins/Runtime/AR/Apple/AppleARKit/Source/AppleARKit/Private/AppleARKitConversion.cpp
--- a/Engine/Plugins/Runtime/AR/Apple/AppleARKit/Source/AppleARKit/Private/AppleARKitConversion.cpp
+++ b/Engine/Plugins/Runtime/AR/Apple/AppleARKit/Source/AppleARKit/Private/AppleARKitConversion.cpp
@@ -3,6 +3,7 @@
 #include "AppleARKitConversion.h"
 #include "AppleARKitModule.h"
 #include "HAL/PlatformMisc.h"
+#include "AppleARKitFaceSupport.h"
 
 template<typename TEnum>
 static FString GetEnumValueAsString(const FString& Name, TEnum Value)
@@ -459,7 +460,22 @@ bool FAppleARKitConversion::IsSessionTrackingFeatureSupported(EARSessionType Ses
 					return [ARWorldTrackingConfiguration supportsFrameSemantics: Semantics];
 				
 				case EARSessionType::Face:
-					return [ARFaceTrackingConfiguration supportsFrameSemantics: Semantics];
+					{
+						static TMap<EARSessionTrackingFeature, bool> SupportFlags;
+						
+						if (!SupportFlags.Contains(SessionTrackingFeature))
+						{
+							SupportFlags.Add(SessionTrackingFeature, false);
+							
+							TArray<IAppleARKitFaceSupport*> Impls = IModularFeatures::Get().GetModularFeatureImplementations<IAppleARKitFaceSupport>("AppleARKitFaceSupport");
+							if (Impls.Num() && Impls[0]->IsARFrameSemanticsSupported(Semantics))
+							{
+								SupportFlags[SessionTrackingFeature] = true;
+							}
+						}
+						
+						return SupportFlags[SessionTrackingFeature];
+					}
 				
 				case EARSessionType::Image:
 					return [ARImageTrackingConfiguration supportsFrameSemantics: Semantics];
diff --git a/Engine/Plugins/Runtime/AR/Apple/AppleARKit/Source/AppleARKit/Public/AppleARKitFaceSupport.h b/Engine/Plugins/Runtime/AR/Apple/AppleARKit/Source/AppleARKit/Public/AppleARKitFaceSupport.h
--- a/Engine/Plugins/Runtime/AR/Apple/AppleARKit/Source/AppleARKit/Public/AppleARKitFaceSupport.h
+++ b/Engine/Plugins/Runtime/AR/Apple/AppleARKit/Source/AppleARKit/Public/AppleARKitFaceSupport.h
@@ -63,7 +63,14 @@ public:
 	 */
 	virtual TArray<FARVideoFormat> ToARConfiguration() { return TArray<FARVideoFormat>(); }
 #endif
-
+	
+#if SUPPORTS_ARKIT_3_0
+	/**
+	 * @return If the desired frame semantics is supported by AR face tracking
+	 */
+	virtual bool IsARFrameSemanticsSupported(ARFrameSemantics InSemantics) const { return false; }
+#endif
+	
 	static FName GetModularFeatureName()
 	{
 		static FName FeatureName = FName(TEXT("AppleARKitFaceSupport"));
diff --git a/Engine/Plugins/Runtime/AR/Apple/AppleARKitFaceSupport/Source/AppleARKitFaceSupport/Private/AppleARKitFaceSupportImpl.cpp b/Engine/Plugins/Runtime/AR/Apple/AppleARKitFaceSupport/Source/AppleARKitFaceSupport/Private/AppleARKitFaceSupportImpl.cpp
--- a/Engine/Plugins/Runtime/AR/Apple/AppleARKitFaceSupport/Source/AppleARKitFaceSupport/Private/AppleARKitFaceSupportImpl.cpp
+++ b/Engine/Plugins/Runtime/AR/Apple/AppleARKitFaceSupport/Source/AppleARKitFaceSupport/Private/AppleARKitFaceSupportImpl.cpp
@@ -246,3 +246,10 @@ TArray<FARVideoFormat> FAppleARKitFaceSupport::ToARConfiguration()
 	return FAppleARKitConversion::FromARVideoFormatArray(ARFaceTrackingConfiguration.supportedVideoFormats);
 }
 #endif
+
+#if SUPPORTS_ARKIT_3_0
+bool FAppleARKitFaceSupport::IsARFrameSemanticsSupported(ARFrameSemantics InSemantics) const
+{
+	return [ARFaceTrackingConfiguration supportsFrameSemantics: InSemantics];
+}
+#endif
diff --git a/Engine/Plugins/Runtime/AR/Apple/AppleARKitFaceSupport/Source/AppleARKitFaceSupport/Private/AppleARKitFaceSupportImpl.h b/Engine/Plugins/Runtime/AR/Apple/AppleARKitFaceSupport/Source/AppleARKitFaceSupport/Private/AppleARKitFaceSupportImpl.h
--- a/Engine/Plugins/Runtime/AR/Apple/AppleARKitFaceSupport/Source/AppleARKitFaceSupport/Private/AppleARKitFaceSupportImpl.h
+++ b/Engine/Plugins/Runtime/AR/Apple/AppleARKitFaceSupport/Source/AppleARKitFaceSupport/Private/AppleARKitFaceSupportImpl.h
@@ -30,6 +30,10 @@ private:
 #if SUPPORTS_ARKIT_1_5
 	virtual TArray<FARVideoFormat> ToARConfiguration() override;
 #endif
+
+#if SUPPORTS_ARKIT_3_0
+	virtual bool IsARFrameSemanticsSupported(ARFrameSemantics InSemantics) const override;
+#endif
 	// ~IAppleARKitFaceSupport
 
 	/** Publishes the remote publisher and the file writer if present */
