Disable ARKit HDR environment texturing until the engine supports it

- The crash happens when the engine code tries to convert the RGBA16F environment texture (from ARKit) to an internal texture of RGBA8 format.
   Apparently the HDR format is a new feature in ARKit 3 that's enabled by default.

diff --git a/Engine/Plugins/Runtime/AR/Apple/AppleARKit/Source/AppleARKit/Private/AppleARKitConversion.cpp b/Engine/Plugins/Runtime/AR/Apple/AppleARKit/Source/AppleARKit/Private/AppleARKitConversion.cpp
--- a/Engine/Plugins/Runtime/AR/Apple/AppleARKit/Source/AppleARKit/Private/AppleARKitConversion.cpp
+++ b/Engine/Plugins/Runtime/AR/Apple/AppleARKit/Source/AppleARKit/Private/AppleARKitConversion.cpp
@@ -312,6 +312,17 @@ ARConfiguration* FAppleARKitConversion::ToARConfiguration( UARSessionConfig* Ses
 				WorldTrackingConfiguration.detectionObjects = ToARReferenceObjectSet(SessionConfig->GetCandidateObjectList(), CandidateObjects);
 			}
 #endif
+			
+#if SUPPORTS_ARKIT_3_0
+			if (FAppleARKitAvailability::SupportsARKit30())
+			{
+				// Disable HDR environment texturing as the engine doesn't support it yet
+				// See FARMetalResource in AppleARKitTextures.cpp
+				// TODO: Add HDR support later
+				WorldTrackingConfiguration.wantsHDREnvironmentTextures = NO;
+			}
+#endif
+			
 			SessionConfiguration = WorldTrackingConfiguration;
 			break;
 		}
