Copying Private-LoadTimes-4.24 stream to Main. Biggest changes are in Materials/Shader memory freezing.

diff --git a/Engine/Plugins/FX/Niagara/Source/NiagaraVertexFactories/Public/NiagaraDrawIndirect.h b/Engine/Plugins/FX/Niagara/Source/NiagaraVertexFactories/Public/NiagaraDrawIndirect.h
--- a/Engine/Plugins/FX/Niagara/Source/NiagaraVertexFactories/Public/NiagaraDrawIndirect.h
+++ b/Engine/Plugins/FX/Niagara/Source/NiagaraVertexFactories/Public/NiagaraDrawIndirect.h
@@ -64,15 +64,15 @@ public:
 	FNiagaraDrawIndirectArgsGenCS() {}
 	FNiagaraDrawIndirectArgsGenCS(const ShaderMetaType::CompiledShaderInitializerType& Initializer);
 
-	bool Serialize(FArchive& Ar);
+	//bool Serialize(FArchive& Ar);
 	void SetOutput(FRHICommandList& RHICmdList, FRHIUnorderedAccessView* DrawIndirectArgsUAV, FRHIUnorderedAccessView* InstanceCountsUAV);
 	void SetParameters(FRHICommandList& RHICmdList, FRHIShaderResourceView* TaskInfosBuffer, int32 NumArgGenTasks, int32 NumInstanceCountClearTasks);
 	void UnbindBuffers(FRHICommandList& RHICmdList);
 
 protected:
 	
-	FShaderResourceParameter TaskInfosParam;
-	FRWShaderParameter InstanceCountsParam;
-	FRWShaderParameter DrawIndirectArgsParam;
-	FShaderParameter TaskCountParam;
+	LAYOUT_FIELD(FShaderResourceParameter, TaskInfosParam)
+	LAYOUT_FIELD(FRWShaderParameter, InstanceCountsParam)
+	LAYOUT_FIELD(FRWShaderParameter, DrawIndirectArgsParam)
+	LAYOUT_FIELD(FShaderParameter, TaskCountParam)
 };
diff --git a/Engine/Plugins/FX/Niagara/Source/NiagaraVertexFactories/Public/NiagaraMeshVertexFactory.h b/Engine/Plugins/FX/Niagara/Source/NiagaraVertexFactories/Public/NiagaraMeshVertexFactory.h
--- a/Engine/Plugins/FX/Niagara/Source/NiagaraVertexFactories/Public/NiagaraMeshVertexFactory.h
+++ b/Engine/Plugins/FX/Niagara/Source/NiagaraVertexFactories/Public/NiagaraMeshVertexFactory.h
@@ -87,16 +87,16 @@ public:
 	/**
 	* Should we cache the material's shadertype on this platform with this vertex factory?
 	*/
-	static bool ShouldCompilePermutation(EShaderPlatform Platform, const class FMaterial* Material, const class FShaderType* ShaderType);
+	static bool ShouldCompilePermutation(const FVertexFactoryShaderPermutationParameters& Parameters);
 
 
 	/**
 	* Modify compile environment to enable instancing
 	* @param OutEnvironment - shader compile environment to modify
 	*/
-	static void ModifyCompilationEnvironment(const FVertexFactoryType* Type, EShaderPlatform Platform, const FMaterial* Material, FShaderCompilerEnvironment& OutEnvironment)
+	static void ModifyCompilationEnvironment(const FVertexFactoryShaderPermutationParameters& Parameters, FShaderCompilerEnvironment& OutEnvironment)
 	{
-		FNiagaraVertexFactoryBase::ModifyCompilationEnvironment(Type, Platform, Material, OutEnvironment);
+		FNiagaraVertexFactoryBase::ModifyCompilationEnvironment(Parameters, OutEnvironment);
 
 		// Set a define so we can tell in MaterialTemplate.usf when we are compiling a mesh particle vertex factory
 		OutEnvironment.SetDefine(TEXT("NIAGARA_MESH_FACTORY"), TEXT("1"));
@@ -171,8 +171,6 @@ public:
 
 	static bool SupportsTessellationShaders() { return true; }
 
-	static FVertexFactoryShaderParameters* ConstructShaderParameters(EShaderFrequency ShaderFrequency);
-	
 	int32 GetLODIndex() const { return LODIndex; }
 	void SetLODIndex(int32 InLODIndex) { LODIndex = InLODIndex; }
 
@@ -213,15 +211,11 @@ public:
 		: FNiagaraMeshVertexFactory()
 	{}
 
-	static bool ShouldCompilePermutation(EShaderPlatform Platform, const class FMaterial* Material, const class FShaderType* ShaderType)
-	{
-		return (Platform == SP_OPENGL_ES2_ANDROID || Platform == SP_OPENGL_ES2_WEBGL) // Those are only platforms that might not support hardware instancing
-			&& FNiagaraMeshVertexFactory::ShouldCompilePermutation(Platform, Material, ShaderType);
-	}
+	static bool ShouldCompilePermutation(const FVertexFactoryShaderPermutationParameters& Parameters);
 
-	static void ModifyCompilationEnvironment(const FVertexFactoryType* Type, EShaderPlatform Platform, const FMaterial* Material, FShaderCompilerEnvironment& OutEnvironment)
+	static void ModifyCompilationEnvironment(const FVertexFactoryShaderPermutationParameters& Parameters, FShaderCompilerEnvironment& OutEnvironment)
 	{
-		FNiagaraMeshVertexFactory::ModifyCompilationEnvironment(Type, Platform, Material, OutEnvironment);
+		FNiagaraMeshVertexFactory::ModifyCompilationEnvironment(Parameters, OutEnvironment);
 
 		OutEnvironment.SetDefine(TEXT("PARTICLE_MESH_INSTANCED"), TEXT("0"));
 	}
diff --git a/Engine/Plugins/FX/Niagara/Source/NiagaraVertexFactories/Public/NiagaraRibbonVertexFactory.h b/Engine/Plugins/FX/Niagara/Source/NiagaraVertexFactories/Public/NiagaraRibbonVertexFactory.h
--- a/Engine/Plugins/FX/Niagara/Source/NiagaraVertexFactories/Public/NiagaraRibbonVertexFactory.h
+++ b/Engine/Plugins/FX/Niagara/Source/NiagaraVertexFactories/Public/NiagaraRibbonVertexFactory.h
@@ -105,12 +105,12 @@ public:
 	/**
 	* Should we cache the material's shadertype on this platform with this vertex factory?
 	*/
-	static bool ShouldCompilePermutation(EShaderPlatform Platform, const class FMaterial* Material, const class FShaderType* ShaderType);
+	static bool ShouldCompilePermutation(const FVertexFactoryShaderPermutationParameters& Parameters);
 
 	/**
 	* Can be overridden by FVertexFactory subclasses to modify their compile environment just before compilation occurs.
 	*/
-	static void ModifyCompilationEnvironment(const FVertexFactoryType* Type, EShaderPlatform Platform, const FMaterial* Material, FShaderCompilerEnvironment& OutEnvironment);
+	static void ModifyCompilationEnvironment(const FVertexFactoryShaderPermutationParameters& Parameters, FShaderCompilerEnvironment& OutEnvironment);
 
 	// FRenderResource interface.
 	virtual void InitRHI() override;
@@ -219,11 +219,6 @@ public:
 		return FacingMode;
 	}
 
-	/**
-	* Construct shader parameters for this type of vertex factory.
-	*/
-	static FVertexFactoryShaderParameters* ConstructShaderParameters(EShaderFrequency ShaderFrequency);
-
 	FIndexBuffer*& GetIndexBuffer()
 	{
 		return IndexBuffer;
diff --git a/Engine/Plugins/FX/Niagara/Source/NiagaraVertexFactories/Public/NiagaraSortingGPU.h b/Engine/Plugins/FX/Niagara/Source/NiagaraVertexFactories/Public/NiagaraSortingGPU.h
--- a/Engine/Plugins/FX/Niagara/Source/NiagaraVertexFactories/Public/NiagaraSortingGPU.h
+++ b/Engine/Plugins/FX/Niagara/Source/NiagaraVertexFactories/Public/NiagaraSortingGPU.h
@@ -44,8 +44,6 @@ public:
 	/** Initialization constructor. */
 	explicit FNiagaraSortKeyGenCS(const ShaderMetaType::CompiledShaderInitializerType& Initializer);
 
-	/** Serialization. */
-	virtual bool Serialize( FArchive& Ar ) override;
 
 	/**
 	 * Set output buffers for this shader.
@@ -64,17 +62,18 @@ public:
 
 private:
 
-	FShaderResourceParameter NiagaraParticleDataFloat;
-	FShaderParameter FloatDataStride;
-	FShaderResourceParameter GPUParticleCountBuffer;
-	FShaderParameter ParticleCountParams;
-	FShaderParameter SortParams;
-	FShaderParameter SortKeyParams;
-	FShaderParameter CameraPosition;
-	FShaderParameter CameraDirection;
+	LAYOUT_FIELD(FShaderResourceParameter, NiagaraParticleDataFloat)
+	LAYOUT_FIELD(FShaderParameter, FloatDataStride)
+	LAYOUT_FIELD(FShaderResourceParameter, GPUParticleCountBuffer)
+	LAYOUT_FIELD(FShaderParameter, ParticleCountParams)
+	LAYOUT_FIELD(FShaderParameter, SortParams)
+	LAYOUT_FIELD(FShaderParameter, SortKeyParams)
+	LAYOUT_FIELD(FShaderParameter, CameraPosition)
+	LAYOUT_FIELD(FShaderParameter, CameraDirection)
 
 	/** Output key buffer. */
-	FShaderResourceParameter OutKeys;
+	LAYOUT_FIELD(FShaderResourceParameter, OutKeys)
 	/** Output indices buffer. */
-	FShaderResourceParameter OutParticleIndices;
+	LAYOUT_FIELD(FShaderResourceParameter, OutParticleIndices)
 };
+
diff --git a/Engine/Plugins/FX/Niagara/Source/NiagaraVertexFactories/Public/NiagaraSpriteVertexFactory.h b/Engine/Plugins/FX/Niagara/Source/NiagaraVertexFactories/Public/NiagaraSpriteVertexFactory.h
--- a/Engine/Plugins/FX/Niagara/Source/NiagaraVertexFactories/Public/NiagaraSpriteVertexFactory.h
+++ b/Engine/Plugins/FX/Niagara/Source/NiagaraVertexFactories/Public/NiagaraSpriteVertexFactory.h
@@ -127,12 +127,12 @@ public:
 	/**
 	 * Should we cache the material's shadertype on this platform with this vertex factory? 
 	 */
-	static bool ShouldCompilePermutation(EShaderPlatform Platform, const class FMaterial* Material, const class FShaderType* ShaderType);
+	static bool ShouldCompilePermutation(const FVertexFactoryShaderPermutationParameters& Parameters);
 
 	/**
 	 * Can be overridden by FVertexFactory subclasses to modify their compile environment just before compilation occurs.
 	 */
-	static void ModifyCompilationEnvironment(const FVertexFactoryType* Type, EShaderPlatform Platform, const FMaterial* Material, FShaderCompilerEnvironment& OutEnvironment);
+	static void ModifyCompilationEnvironment(const FVertexFactoryShaderPermutationParameters& Parameters, FShaderCompilerEnvironment& OutEnvironment);
 	
 	void SetTexCoordBuffer(const FVertexBuffer* InTexCoordBuffer);
 	
@@ -218,11 +218,6 @@ public:
 		VertexBufferOverride = InVertexBufferOverride;
 	}
 
-	/**
-	 * Construct shader parameters for this type of vertex factory.
-	 */
-	static FVertexFactoryShaderParameters* ConstructShaderParameters(EShaderFrequency ShaderFrequency);
-
 	FUniformBufferRHIRef LooseParameterUniformBuffer;
 
 protected:
diff --git a/Engine/Plugins/FX/Niagara/Source/NiagaraVertexFactories/Public/NiagaraVertexFactory.h b/Engine/Plugins/FX/Niagara/Source/NiagaraVertexFactories/Public/NiagaraVertexFactory.h
--- a/Engine/Plugins/FX/Niagara/Source/NiagaraVertexFactories/Public/NiagaraVertexFactory.h
+++ b/Engine/Plugins/FX/Niagara/Source/NiagaraVertexFactories/Public/NiagaraVertexFactory.h
@@ -76,9 +76,9 @@ public:
 		bNeedsDeclaration = false;
 	}
 	
-	static void ModifyCompilationEnvironment(const FVertexFactoryType* Type, EShaderPlatform Platform, const FMaterial* Material, FShaderCompilerEnvironment& OutEnvironment)
+	static void ModifyCompilationEnvironment(const FVertexFactoryShaderPermutationParameters& Parameters, FShaderCompilerEnvironment& OutEnvironment)
 	{
-		FVertexFactory::ModifyCompilationEnvironment(Type, Platform, Material, OutEnvironment);
+		FVertexFactory::ModifyCompilationEnvironment(Parameters, OutEnvironment);
 		OutEnvironment.SetDefine(TEXT("NIAGARA_PARTICLE_FACTORY"), TEXT("1"));
 	}
 
diff --git a/Engine/Plugins/Lumin/MagicLeap/Source/MagicLeap/Private/MagicLeapCustomPresent.cpp b/Engine/Plugins/Lumin/MagicLeap/Source/MagicLeap/Private/MagicLeapCustomPresent.cpp
--- a/Engine/Plugins/Lumin/MagicLeap/Source/MagicLeap/Private/MagicLeapCustomPresent.cpp
+++ b/Engine/Plugins/Lumin/MagicLeap/Source/MagicLeap/Private/MagicLeapCustomPresent.cpp
@@ -234,8 +234,8 @@ void FMagicLeapCustomPresent::RenderToTextureSlice_RenderThread(FRHICommandListI
 		GraphicsPSOInit.RasterizerState = TStaticRasterizerState<>::GetRHI();
 		GraphicsPSOInit.DepthStencilState = TStaticDepthStencilState<false, CF_Always>::GetRHI();
 		GraphicsPSOInit.BoundShaderState.VertexDeclarationRHI = GFilterVertexDeclaration.VertexDeclarationRHI;
-		GraphicsPSOInit.BoundShaderState.VertexShaderRHI = GETSAFERHISHADER_VERTEX(*VertexShader);
-		GraphicsPSOInit.BoundShaderState.PixelShaderRHI = GETSAFERHISHADER_PIXEL(*PixelShader);
+		GraphicsPSOInit.BoundShaderState.VertexShaderRHI = VertexShader.GetVertexShader();
+		GraphicsPSOInit.BoundShaderState.PixelShaderRHI = PixelShader.GetPixelShader();
 		GraphicsPSOInit.PrimitiveType = PT_TriangleList;
 
 		SetGraphicsPipelineState(RHICmdList, GraphicsPSOInit);
@@ -250,7 +250,7 @@ void FMagicLeapCustomPresent::RenderToTextureSlice_RenderThread(FRHICommandListI
 			UVandSize.Z, UVandSize.W,	// SizeU, SizeV
 			FIntPoint(vp_width, vp_height),
 			FIntPoint(1, 1),
-			*VertexShader,
+			VertexShader,
 			EDRF_Default);
 	}
 	RHICmdList.EndRenderPass();
diff --git a/Engine/Plugins/Lumin/MagicLeap/Source/MagicLeap/Private/MagicLeapHMD.cpp b/Engine/Plugins/Lumin/MagicLeap/Source/MagicLeap/Private/MagicLeapHMD.cpp
--- a/Engine/Plugins/Lumin/MagicLeap/Source/MagicLeap/Private/MagicLeapHMD.cpp
+++ b/Engine/Plugins/Lumin/MagicLeap/Source/MagicLeap/Private/MagicLeapHMD.cpp
@@ -1929,8 +1929,8 @@ void FMagicLeapHMD::RenderTexture_RenderThread(class FRHICommandListImmediate& R
 			GraphicsPSOInit.RasterizerState = TStaticRasterizerState<>::GetRHI();
 			GraphicsPSOInit.DepthStencilState = TStaticDepthStencilState<false, CF_Always>::GetRHI();
 			GraphicsPSOInit.BoundShaderState.VertexDeclarationRHI = GFilterVertexDeclaration.VertexDeclarationRHI;
-			GraphicsPSOInit.BoundShaderState.VertexShaderRHI = GETSAFERHISHADER_VERTEX(*VertexShader);
-			GraphicsPSOInit.BoundShaderState.PixelShaderRHI = GETSAFERHISHADER_PIXEL(*PixelShader);
+			GraphicsPSOInit.BoundShaderState.VertexShaderRHI = VertexShader.GetVertexShader();
+			GraphicsPSOInit.BoundShaderState.PixelShaderRHI = PixelShader.GetPixelShader();
 			GraphicsPSOInit.PrimitiveType = PT_TriangleList;
 
 			SetGraphicsPipelineState(RHICmdList, GraphicsPSOInit);
@@ -1947,7 +1947,7 @@ void FMagicLeapHMD::RenderTexture_RenderThread(class FRHICommandListImmediate& R
 					0.5f, 1.0f,
 					FIntPoint(ViewportWidth, ViewportHeight),
 					FIntPoint(1, 1),
-					*VertexShader,
+					VertexShader,
 					EDRF_Default);
 			}
 			else if (WindowMirrorMode == 2)
@@ -1960,7 +1960,7 @@ void FMagicLeapHMD::RenderTexture_RenderThread(class FRHICommandListImmediate& R
 					1.0f, 1.0f,
 					FIntPoint(ViewportWidth, ViewportHeight),
 					FIntPoint(1, 1),
-					*VertexShader,
+					VertexShader,
 					EDRF_Default);
 			}
 		}
diff --git a/Engine/Plugins/Media/AvfMedia/Source/AvfMedia/Private/Player/AvfMediaVideoSampler.cpp b/Engine/Plugins/Media/AvfMedia/Source/AvfMedia/Private/Player/AvfMediaVideoSampler.cpp
--- a/Engine/Plugins/Media/AvfMedia/Source/AvfMedia/Private/Player/AvfMediaVideoSampler.cpp
+++ b/Engine/Plugins/Media/AvfMedia/Source/AvfMedia/Private/Player/AvfMediaVideoSampler.cpp
@@ -305,8 +305,8 @@ void FAvfMediaVideoSampler::ProcessFrame(CVPixelBufferRef Frame, FTimespan Sampl
 					GraphicsPSOInit.DepthStencilState = TStaticDepthStencilState<false, CF_Always>::GetRHI();
 
 					GraphicsPSOInit.BoundShaderState.VertexDeclarationRHI = GMediaVertexDeclaration.VertexDeclarationRHI;
-					GraphicsPSOInit.BoundShaderState.VertexShaderRHI = GETSAFERHISHADER_VERTEX(*VertexShader);
-					GraphicsPSOInit.BoundShaderState.PixelShaderRHI = GETSAFERHISHADER_PIXEL(*PixelShader);
+					GraphicsPSOInit.BoundShaderState.VertexShaderRHI = VertexShader.GetVertexShader();
+					GraphicsPSOInit.BoundShaderState.PixelShaderRHI = PixelShader.GetPixelShader();
 					GraphicsPSOInit.PrimitiveType = PT_TriangleStrip;
 
 					SetGraphicsPipelineState(RHICmdList, GraphicsPSOInit);
diff --git a/Engine/Plugins/Media/MediaIOFramework/Source/MediaIOCore/Private/MediaCapture.cpp b/Engine/Plugins/Media/MediaIOFramework/Source/MediaIOCore/Private/MediaCapture.cpp
--- a/Engine/Plugins/Media/MediaIOFramework/Source/MediaIOCore/Private/MediaCapture.cpp
+++ b/Engine/Plugins/Media/MediaIOFramework/Source/MediaIOCore/Private/MediaCapture.cpp
@@ -828,7 +828,7 @@ void UMediaCapture::Capture_RenderThread(FRHICommandListImmediate& RHICmdList,
 				TShaderMapRef<FMediaShadersVS> VertexShader(ShaderMap);
 
 				GraphicsPSOInit.BoundShaderState.VertexDeclarationRHI = GMediaVertexDeclaration.VertexDeclarationRHI;
-				GraphicsPSOInit.BoundShaderState.VertexShaderRHI = GETSAFERHISHADER_VERTEX(*VertexShader);
+				GraphicsPSOInit.BoundShaderState.VertexShaderRHI = VertexShader.GetVertexShader();
 
 				const bool bDoLinearToSRGB = false;
 
@@ -837,7 +837,7 @@ void UMediaCapture::Capture_RenderThread(FRHICommandListImmediate& RHICmdList,
 				case EMediaCaptureConversionOperation::RGBA8_TO_YUV_8BIT:
 				{
 					TShaderMapRef<FRGB8toUYVY8ConvertPS> ConvertShader(ShaderMap);
-					GraphicsPSOInit.BoundShaderState.PixelShaderRHI = GETSAFERHISHADER_PIXEL(*ConvertShader);
+					GraphicsPSOInit.BoundShaderState.PixelShaderRHI = ConvertShader.GetPixelShader();
 					SetGraphicsPipelineState(RHICmdList, GraphicsPSOInit);
 					ConvertShader->SetParameters(RHICmdList, SourceTexture, MediaShaders::RgbToYuvRec709Full, MediaShaders::YUVOffset8bits, bDoLinearToSRGB);
 				}
@@ -845,7 +845,7 @@ void UMediaCapture::Capture_RenderThread(FRHICommandListImmediate& RHICmdList,
 				case EMediaCaptureConversionOperation::RGB10_TO_YUVv210_10BIT:
 				{
 					TShaderMapRef<FRGB10toYUVv210ConvertPS> ConvertShader(ShaderMap);
-					GraphicsPSOInit.BoundShaderState.PixelShaderRHI = GETSAFERHISHADER_PIXEL(*ConvertShader);
+					GraphicsPSOInit.BoundShaderState.PixelShaderRHI = ConvertShader.GetPixelShader();
 					SetGraphicsPipelineState(RHICmdList, GraphicsPSOInit);
 					ConvertShader->SetParameters(RHICmdList, SourceTexture, MediaShaders::RgbToYuvRec709Full, MediaShaders::YUVOffset10bits, bDoLinearToSRGB);
 				}
@@ -853,7 +853,7 @@ void UMediaCapture::Capture_RenderThread(FRHICommandListImmediate& RHICmdList,
 				case EMediaCaptureConversionOperation::INVERT_ALPHA:
 				{
 					TShaderMapRef<FInvertAlphaPS> ConvertShader(ShaderMap);
-					GraphicsPSOInit.BoundShaderState.PixelShaderRHI = GETSAFERHISHADER_PIXEL(*ConvertShader);
+					GraphicsPSOInit.BoundShaderState.PixelShaderRHI = ConvertShader.GetPixelShader();
 					SetGraphicsPipelineState(RHICmdList, GraphicsPSOInit);
 					ConvertShader->SetParameters(RHICmdList, SourceTexture);
 				}
@@ -861,7 +861,7 @@ void UMediaCapture::Capture_RenderThread(FRHICommandListImmediate& RHICmdList,
 				case EMediaCaptureConversionOperation::SET_ALPHA_ONE:
 				{
 					TShaderMapRef<FSetAlphaOnePS> ConvertShader(ShaderMap);
-					GraphicsPSOInit.BoundShaderState.PixelShaderRHI = GETSAFERHISHADER_PIXEL(*ConvertShader);
+					GraphicsPSOInit.BoundShaderState.PixelShaderRHI = ConvertShader.GetPixelShader();
 					SetGraphicsPipelineState(RHICmdList, GraphicsPSOInit);
 					ConvertShader->SetParameters(RHICmdList, SourceTexture);
 				}
diff --git a/Engine/Plugins/Media/WebMMedia/Source/WebMMedia/Private/Player/WebMVideoDecoder.cpp b/Engine/Plugins/Media/WebMMedia/Source/WebMMedia/Private/Player/WebMVideoDecoder.cpp
--- a/Engine/Plugins/Media/WebMMedia/Source/WebMMedia/Private/Player/WebMVideoDecoder.cpp
+++ b/Engine/Plugins/Media/WebMMedia/Source/WebMMedia/Private/Player/WebMVideoDecoder.cpp
@@ -262,8 +262,8 @@ void FWebMVideoDecoder::ConvertYUVToRGBAndSubmit(const FConvertParams& Params)
 				GraphicsPSOInit.RasterizerState = TStaticRasterizerState<>::GetRHI();
 				GraphicsPSOInit.DepthStencilState = TStaticDepthStencilState<false, CF_Always>::GetRHI();
 				GraphicsPSOInit.BoundShaderState.VertexDeclarationRHI = GMoviePlayerResources.VertexDeclarationRHI;
-				GraphicsPSOInit.BoundShaderState.VertexShaderRHI = GETSAFERHISHADER_VERTEX(*VertexShader);
-				GraphicsPSOInit.BoundShaderState.PixelShaderRHI = GETSAFERHISHADER_PIXEL(*PixelShader);
+				GraphicsPSOInit.BoundShaderState.VertexShaderRHI = VertexShader.GetVertexShader();
+				GraphicsPSOInit.BoundShaderState.PixelShaderRHI = PixelShader.GetPixelShader();
 				GraphicsPSOInit.PrimitiveType = PT_TriangleStrip;
 			}
 
diff --git a/Engine/Plugins/Media/WmfMedia/Source/WmfMedia/Private/Wmf/WmfMediaHardwareVideoDecodingRendering.cpp b/Engine/Plugins/Media/WmfMedia/Source/WmfMedia/Private/Wmf/WmfMediaHardwareVideoDecodingRendering.cpp
--- a/Engine/Plugins/Media/WmfMedia/Source/WmfMedia/Private/Wmf/WmfMediaHardwareVideoDecodingRendering.cpp
+++ b/Engine/Plugins/Media/WmfMedia/Source/WmfMedia/Private/Wmf/WmfMediaHardwareVideoDecodingRendering.cpp
@@ -159,48 +159,48 @@ bool FWmfMediaHardwareVideoDecodingParameters::ConvertTextureFormat_RenderThread
 			new (RHICmdList.AllocCommand<FRHICommandCopyResource>()) FRHICommandCopyResource(SampleTexture, SampleDestinationTexture);
 		}
 
-		TShaderMap<FGlobalShaderType>* GlobalShaderMap = GetGlobalShaderMap(GMaxRHIFeatureLevel);
+		FGlobalShaderMap* GlobalShaderMap = GetGlobalShaderMap(GMaxRHIFeatureLevel);
 
 
 		if (InSample->GetFormat() == EMediaTextureSampleFormat::CharNV12)
 		{
 			TShaderMapRef< FHardwareVideoDecodingVS > VertexShader(GlobalShaderMap);
 			TShaderMapRef< FHardwareVideoDecodingPS > PixelShader(GlobalShaderMap);
-			GraphicsPSOInit.BoundShaderState.VertexShaderRHI = GETSAFERHISHADER_VERTEX(*VertexShader);
-			GraphicsPSOInit.BoundShaderState.PixelShaderRHI = GETSAFERHISHADER_PIXEL(*PixelShader);
+			GraphicsPSOInit.BoundShaderState.VertexShaderRHI = VertexShader.GetVertexShader();
+			GraphicsPSOInit.BoundShaderState.PixelShaderRHI = PixelShader.GetPixelShader();
 
 			SetGraphicsPipelineState(RHICmdList, GraphicsPSOInit);
 
 			FShaderResourceViewRHIRef Y_SRV = RHICreateShaderResourceView(SampleDestinationTexture, 0, 1, PF_G8);
 			FShaderResourceViewRHIRef UV_SRV = RHICreateShaderResourceView(SampleDestinationTexture, 0, 1, PF_R8G8);
-			VertexShader->SetParameters(RHICmdList, VertexShader->GetVertexShader(), Y_SRV, UV_SRV, InSample->IsOutputSrgb());
-			PixelShader->SetParameters(RHICmdList, PixelShader->GetPixelShader(), Y_SRV, UV_SRV, InSample->IsOutputSrgb());
+			VertexShader->SetParameters(RHICmdList, VertexShader.GetVertexShader(), Y_SRV, UV_SRV, InSample->IsOutputSrgb());
+			PixelShader->SetParameters(RHICmdList, PixelShader.GetPixelShader(), Y_SRV, UV_SRV, InSample->IsOutputSrgb());
 		}
 		else if (InSample->GetFormat() == EMediaTextureSampleFormat::CharBGRA)
 		{
 			TShaderMapRef< FHardwareVideoDecodingVS > VertexShader(GlobalShaderMap);
 			TShaderMapRef< FHardwareVideoDecodingPassThroughPS > PixelShader(GlobalShaderMap);
-			GraphicsPSOInit.BoundShaderState.VertexShaderRHI = GETSAFERHISHADER_VERTEX(*VertexShader);
-			GraphicsPSOInit.BoundShaderState.PixelShaderRHI = GETSAFERHISHADER_PIXEL(*PixelShader);
+			GraphicsPSOInit.BoundShaderState.VertexShaderRHI = VertexShader.GetVertexShader();
+			GraphicsPSOInit.BoundShaderState.PixelShaderRHI = PixelShader.GetPixelShader();
 
 			SetGraphicsPipelineState(RHICmdList, GraphicsPSOInit);
 
 			FShaderResourceViewRHIRef SRV = RHICreateShaderResourceView(SampleDestinationTexture, 0, 1, PF_B8G8R8A8);
-			VertexShader->SetParameters(RHICmdList, VertexShader->GetVertexShader(), SRV, InSample->IsOutputSrgb());
-			PixelShader->SetParameters(RHICmdList, PixelShader->GetPixelShader(), SRV, InSample->IsOutputSrgb());
+			VertexShader->SetParameters(RHICmdList, VertexShader.GetVertexShader(), SRV, InSample->IsOutputSrgb());
+			PixelShader->SetParameters(RHICmdList, PixelShader.GetPixelShader(), SRV, InSample->IsOutputSrgb());
 		}
 		else if (InSample->GetFormat() == EMediaTextureSampleFormat::Y416)
 		{
 			TShaderMapRef< FHardwareVideoDecodingVS > VertexShader(GlobalShaderMap);
 			TShaderMapRef< FHardwareVideoDecodingY416PS > PixelShader(GlobalShaderMap);
-			GraphicsPSOInit.BoundShaderState.VertexShaderRHI = GETSAFERHISHADER_VERTEX(*VertexShader);
-			GraphicsPSOInit.BoundShaderState.PixelShaderRHI = GETSAFERHISHADER_PIXEL(*PixelShader);
+			GraphicsPSOInit.BoundShaderState.VertexShaderRHI = VertexShader.GetVertexShader();
+			GraphicsPSOInit.BoundShaderState.PixelShaderRHI = PixelShader.GetPixelShader();
 
 			SetGraphicsPipelineState(RHICmdList, GraphicsPSOInit);
 
 			FShaderResourceViewRHIRef SRV = RHICreateShaderResourceView(SampleDestinationTexture, 0, 1, PF_A16B16G16R16);
-			VertexShader->SetParameters(RHICmdList, VertexShader->GetVertexShader(), SRV, InSample->IsOutputSrgb());
-			PixelShader->SetParameters(RHICmdList, PixelShader->GetPixelShader(), SRV, InSample->IsOutputSrgb());
+			VertexShader->SetParameters(RHICmdList, VertexShader.GetVertexShader(), SRV, InSample->IsOutputSrgb());
+			PixelShader->SetParameters(RHICmdList, PixelShader.GetPixelShader(), SRV, InSample->IsOutputSrgb());
 		}
 
 		RHICmdList.DrawPrimitive(0, 2, 1);
diff --git a/Engine/Plugins/Media/WmfMedia/Source/WmfMedia/Private/Wmf/WmfMediaHardwareVideoDecodingShaders.cpp b/Engine/Plugins/Media/WmfMedia/Source/WmfMedia/Private/Wmf/WmfMediaHardwareVideoDecodingShaders.cpp
--- a/Engine/Plugins/Media/WmfMedia/Source/WmfMedia/Private/Wmf/WmfMediaHardwareVideoDecodingShaders.cpp
+++ b/Engine/Plugins/Media/WmfMedia/Source/WmfMedia/Private/Wmf/WmfMediaHardwareVideoDecodingShaders.cpp
@@ -2,6 +2,8 @@
 
 #include "WmfMediaHardwareVideoDecodingShaders.h"
 
+IMPLEMENT_TYPE_LAYOUT(FWmfMediaHardwareVideoDecodingShader);
+
 IMPLEMENT_SHADER_TYPE(, FHardwareVideoDecodingVS, TEXT("/Plugin/WmfMedia/Private/MediaHardwareVideoDecoding.usf"), TEXT("MainVS"), SF_Vertex)
 IMPLEMENT_SHADER_TYPE(, FHardwareVideoDecodingPS, TEXT("/Plugin/WmfMedia/Private/MediaHardwareVideoDecoding.usf"), TEXT("NV12ConvertPS"), SF_Pixel)
 IMPLEMENT_SHADER_TYPE(, FHardwareVideoDecodingPassThroughPS, TEXT("/Plugin/WmfMedia/Private/MediaHardwareVideoDecoding.usf"), TEXT("PassThroughPS"), SF_Pixel)
diff --git a/Engine/Plugins/Media/WmfMedia/Source/WmfMedia/Public/WmfMediaHardwareVideoDecodingShaders.h b/Engine/Plugins/Media/WmfMedia/Source/WmfMedia/Public/WmfMediaHardwareVideoDecodingShaders.h
--- a/Engine/Plugins/Media/WmfMedia/Source/WmfMedia/Public/WmfMediaHardwareVideoDecodingShaders.h
+++ b/Engine/Plugins/Media/WmfMedia/Source/WmfMedia/Public/WmfMediaHardwareVideoDecodingShaders.h
@@ -14,6 +14,7 @@
  */
 class FWmfMediaHardwareVideoDecodingShader : public FGlobalShader
 {
+	DECLARE_TYPE_LAYOUT(FWmfMediaHardwareVideoDecodingShader, NonVirtual);
 public:
 	static bool ShouldCompilePermutation(const FGlobalShaderPermutationParameters& Parameters)
 	{
@@ -77,21 +78,14 @@ public:
 		SetShaderValue(RHICmdList, ShaderRHI, SrgbToLinear, InIsOutputSrgb ? 1 : 0); // Explicitly specify integer value, as using boolean falls over on XboxOne.
 	}
 
-
-	virtual bool Serialize(FArchive& Ar) override
-	{
-		bool bShaderHasOutdatedParameters = FGlobalShader::Serialize(Ar);
-		Ar << TextureY << TextureUV << PointClampedSamplerY << BilinearClampedSamplerUV << ColorTransform << SrgbToLinear;
-		return bShaderHasOutdatedParameters;
-	}
-
 private:
-	FShaderResourceParameter TextureY;
-	FShaderResourceParameter TextureUV;
-	FShaderResourceParameter PointClampedSamplerY;
-	FShaderResourceParameter BilinearClampedSamplerUV;
-	FShaderParameter ColorTransform;
-	FShaderParameter SrgbToLinear;
+	
+	LAYOUT_FIELD(FShaderResourceParameter, TextureY);
+	LAYOUT_FIELD(FShaderResourceParameter, TextureUV);
+	LAYOUT_FIELD(FShaderResourceParameter, PointClampedSamplerY);
+	LAYOUT_FIELD(FShaderResourceParameter, BilinearClampedSamplerUV);
+	LAYOUT_FIELD(FShaderParameter, ColorTransform);
+	LAYOUT_FIELD(FShaderParameter, SrgbToLinear);
 };
 
 class FHardwareVideoDecodingVS : public FWmfMediaHardwareVideoDecodingShader
diff --git a/Engine/Plugins/NetcodeUnitTest/NetcodeUnitTest/Source/NetcodeUnitTest/NetcodeUnitTest.Build.cs b/Engine/Plugins/NetcodeUnitTest/NetcodeUnitTest/Source/NetcodeUnitTest/NetcodeUnitTest.Build.cs
--- a/Engine/Plugins/NetcodeUnitTest/NetcodeUnitTest/Source/NetcodeUnitTest/NetcodeUnitTest.Build.cs
+++ b/Engine/Plugins/NetcodeUnitTest/NetcodeUnitTest/Source/NetcodeUnitTest/NetcodeUnitTest.Build.cs
@@ -27,6 +27,7 @@ namespace UnrealBuildTool.Rules
 					"EngineSettings",
 					"InputCore",
 					"PacketHandler",
+					"RenderCore",
 					"Slate",
 					"SlateCore"
 				}
diff --git a/Engine/Plugins/Runtime/AR/Apple/AppleARKit/Source/AppleARKit/Private/AppleARKitVideoOverlay.cpp b/Engine/Plugins/Runtime/AR/Apple/AppleARKit/Source/AppleARKit/Private/AppleARKitVideoOverlay.cpp
--- a/Engine/Plugins/Runtime/AR/Apple/AppleARKit/Source/AppleARKit/Private/AppleARKitVideoOverlay.cpp
+++ b/Engine/Plugins/Runtime/AR/Apple/AppleARKit/Source/AppleARKit/Private/AppleARKitVideoOverlay.cpp
@@ -85,20 +85,20 @@ public:
 	{
 		if (bIsMobileRenderer)
 		{
-			return Parameters.Material->GetMaterialDomain() == MD_PostProcess && IsMobilePlatform(Parameters.Platform);
+			return Parameters.MaterialParameters.MaterialDomain == MD_PostProcess && IsMobilePlatform(Parameters.Platform);
 		}
 		else
 		{
-			return Parameters.Material->GetMaterialDomain() == MD_PostProcess && !IsMobilePlatform(Parameters.Platform);
+			return Parameters.MaterialParameters.MaterialDomain == MD_PostProcess && !IsMobilePlatform(Parameters.Platform);
 		}
 	}
 
 	static void ModifyCompilationEnvironment(const FMaterialShaderPermutationParameters& Parameters, FShaderCompilerEnvironment& OutEnvironment)
 	{
-		FMaterialShader::ModifyCompilationEnvironment(Parameters.Platform, OutEnvironment);
+		FMaterialShader::ModifyCompilationEnvironment(Parameters, OutEnvironment);
 		OutEnvironment.SetDefine(TEXT("POST_PROCESS_MATERIAL"), 1);
 		OutEnvironment.SetDefine(TEXT("POST_PROCESS_MATERIAL_MOBILE"), bIsMobileRenderer);
-		OutEnvironment.SetDefine(TEXT("POST_PROCESS_MATERIAL_BEFORE_TONEMAP"), (Parameters.Material->GetBlendableLocation() != BL_AfterTonemapping) ? 1 : 0);
+		OutEnvironment.SetDefine(TEXT("POST_PROCESS_MATERIAL_BEFORE_TONEMAP"), (Parameters.MaterialParameters.BlendableLocation != BL_AfterTonemapping) ? 1 : 0);
 	}
 };
 
@@ -108,7 +108,7 @@ class TARKitCameraOverlayVS : public TPostProcessMaterialShader<bIsMobileRendere
 {
 public:
 	using FMaterialShader::FPermutationDomain;
-	DECLARE_MATERIAL_SHADER(TARKitCameraOverlayVS);
+	DECLARE_SHADER_TYPE(TARKitCameraOverlayVS, Material);
 
 	static void ModifyCompilationEnvironment(const FMaterialShaderPermutationParameters& Parameters, FShaderCompilerEnvironment& OutEnvironment)
 	{
@@ -123,7 +123,7 @@ public:
 
 	void SetParameters(FRHICommandList& RHICmdList, const FSceneView View)
 	{
-		FRHIVertexShader* ShaderRHI = this->GetVertexShader();
+		FRHIVertexShader* ShaderRHI = RHICmdList.GetBoundVertexShader();
 		TPostProcessMaterialShader<bIsMobileRenderer>::SetViewParameters(RHICmdList, ShaderRHI, View, View.ViewUniformBuffer);
 	}
 };
@@ -131,15 +131,15 @@ public:
 using FARKitCameraOverlayVS = TARKitCameraOverlayVS<false>;
 using FARKitCameraOverlayMobileVS = TARKitCameraOverlayVS<true>;
 
-template<> IMPLEMENT_MATERIAL_SHADER(FARKitCameraOverlayVS, "/Engine/Private/PostProcessMaterialShaders.usf", "MainVS_VideoOverlay", SF_Vertex);
-template<> IMPLEMENT_MATERIAL_SHADER(FARKitCameraOverlayMobileVS, "/Engine/Private/PostProcessMaterialShaders.usf", "MainVS", SF_Vertex);
+IMPLEMENT_SHADER_TYPE(template<>, FARKitCameraOverlayVS, TEXT("/Engine/Private/PostProcessMaterialShaders.usf"), TEXT("MainVS_VideoOverlay"), SF_Vertex);
+IMPLEMENT_SHADER_TYPE(template<>, FARKitCameraOverlayMobileVS, TEXT("/Engine/Private/PostProcessMaterialShaders.usf"), TEXT("MainVS"), SF_Vertex);
 
 template <bool bIsMobileRenderer>
 class TARKitCameraOverlayPS : public TPostProcessMaterialShader<bIsMobileRenderer>
 {
 public:
 	using FMaterialShader::FPermutationDomain;
-	DECLARE_MATERIAL_SHADER(TARKitCameraOverlayPS);
+	DECLARE_SHADER_TYPE(TARKitCameraOverlayPS, Material);
 
 	static void ModifyCompilationEnvironment(const FMaterialShaderPermutationParameters& Parameters, FShaderCompilerEnvironment& OutEnvironment)
 	{
@@ -154,7 +154,7 @@ public:
 
 	void SetParameters(FRHICommandList& RHICmdList, const FSceneView View, const FMaterialRenderProxy* Material)
 	{
-		FRHIPixelShader* ShaderRHI = this->GetPixelShader();
+		FRHIPixelShader* ShaderRHI = RHICmdList.GetBoundPixelShader();
 		TPostProcessMaterialShader<bIsMobileRenderer>::SetParameters(RHICmdList, ShaderRHI, Material, *Material->GetMaterial(View.GetFeatureLevel()), View, View.ViewUniformBuffer, ESceneTextureSetupMode::None);
 	}
 };
@@ -162,8 +162,8 @@ public:
 using FARKitCameraOverlayPS = TARKitCameraOverlayPS<false>;
 using FARKitCameraOverlayMobilePS = TARKitCameraOverlayPS<true>;
 
-template<> IMPLEMENT_MATERIAL_SHADER(FARKitCameraOverlayPS, "/Engine/Private/PostProcessMaterialShaders.usf", "MainPS_VideoOverlay", SF_Pixel);
-template<> IMPLEMENT_MATERIAL_SHADER(FARKitCameraOverlayMobilePS, "/Engine/Private/PostProcessMaterialShaders.usf", "MainPS", SF_Pixel);
+IMPLEMENT_SHADER_TYPE(template<>, FARKitCameraOverlayPS, TEXT("/Engine/Private/PostProcessMaterialShaders.usf"), TEXT("MainPS_VideoOverlay"), SF_Pixel);
+IMPLEMENT_SHADER_TYPE(template<>, FARKitCameraOverlayMobilePS, TEXT("/Engine/Private/PostProcessMaterialShaders.usf"), TEXT("MainPS"), SF_Pixel);
 
 void FAppleARKitVideoOverlay::UpdateOcclusionTextures(const FAppleARKitFrame& Frame)
 {
@@ -178,19 +178,19 @@ void FAppleARKitVideoOverlay::UpdateOcclusionTextures(const FAppleARKitFrame& Fr
 		{
 			id<MTLDevice> Device = (id<MTLDevice>)GDynamicRHI->RHIGetNativeDevice();
 			check(Device);
-			
+
 			if (!MatteGenerator)
 			{
 				// TODO: add config variable for the matte resolution
 				MatteGenerator = [[ARMatteGenerator alloc] initWithDevice: Device
 														  matteResolution: ARMatteResolutionFull];
 			}
-			
+
 			if (!CommandQueue)
 			{
 				CommandQueue = [Device newCommandQueue];
 			}
-			
+
 			id<MTLCommandBuffer> CommandBuffer = [CommandQueue commandBuffer];
 			id<MTLTexture> MatteTexture = nullptr;
 			id<MTLTexture> DepthTexture = nullptr;
@@ -202,13 +202,13 @@ void FAppleARKitVideoOverlay::UpdateOcclusionTextures(const FAppleARKitFrame& Fr
 			{
 				OcclusionMatteTexture->SetMetalTexture(Frame.Timestamp, MatteTexture);
 			}
-			
+
 			if (DepthTexture && OcclusionDepthTexture)
 			{
 				OcclusionDepthTexture->SetMetalTexture(Frame.Timestamp, DepthTexture);
 				bOcclusionDepthTextureRecentlyUpdated = true;
 			}
-			
+
 			[CommandBuffer commit];
 		}
 	}
@@ -221,7 +221,7 @@ void FAppleARKitVideoOverlay::RenderVideoOverlayWithMaterial(FRHICommandListImme
 	{
 		return;
 	}
-	
+
 #if STATS && PLATFORM_IOS
 	// LastUpdateTimestamp has been removed
 	//const auto DelayMS = ([[NSProcessInfo processInfo] systemUptime] - LastUpdateTimestamp) * 1000.0;
@@ -331,7 +331,7 @@ void FAppleARKitVideoOverlay::RenderVideoOverlayWithMaterial(FRHICommandListImme
 
 	FGraphicsPipelineStateInitializer GraphicsPSOInit;
 	RHICmdList.ApplyCachedRenderTargets(GraphicsPSOInit);
-	
+
 	GraphicsPSOInit.RasterizerState = TStaticRasterizerState<>::GetRHI();
 	if (bRenderingOcclusion)
 	{
@@ -348,25 +348,25 @@ void FAppleARKitVideoOverlay::RenderVideoOverlayWithMaterial(FRHICommandListImme
 	GraphicsPSOInit.BoundShaderState.VertexDeclarationRHI = GFilterVertexDeclaration.VertexDeclarationRHI;
 
 	const bool bIsMobileRenderer = FeatureLevel <= ERHIFeatureLevel::ES3_1;
-	FMaterialShader* VertexShader = nullptr;
-	FMaterialShader* PixelShader = nullptr;
+	TShaderRef<FMaterialShader> VertexShader;
+	TShaderRef<FMaterialShader> PixelShader;
 	if (bIsMobileRenderer)
 	{
 		VertexShader = MaterialShaderMap->GetShader<FARKitCameraOverlayMobileVS>();
-		check(VertexShader != nullptr);
+		check(VertexShader.IsValid());
 		PixelShader = MaterialShaderMap->GetShader<FARKitCameraOverlayMobilePS>();
-		check(PixelShader != nullptr);
-		GraphicsPSOInit.BoundShaderState.VertexShaderRHI = GETSAFERHISHADER_VERTEX(VertexShader);
-		GraphicsPSOInit.BoundShaderState.PixelShaderRHI = GETSAFERHISHADER_PIXEL(PixelShader);
+		check(PixelShader.IsValid());
+		GraphicsPSOInit.BoundShaderState.VertexShaderRHI = VertexShader.GetVertexShader();
+		GraphicsPSOInit.BoundShaderState.PixelShaderRHI = PixelShader.GetPixelShader();
 	}
 	else
 	{
 		VertexShader = MaterialShaderMap->GetShader<FARKitCameraOverlayVS>();
-		check(VertexShader != nullptr);
+		check(VertexShader.IsValid());
 		PixelShader = MaterialShaderMap->GetShader<FARKitCameraOverlayPS>();
-		check(PixelShader != nullptr);
-		GraphicsPSOInit.BoundShaderState.VertexShaderRHI = GETSAFERHISHADER_VERTEX(VertexShader);
-		GraphicsPSOInit.BoundShaderState.PixelShaderRHI = GETSAFERHISHADER_PIXEL(PixelShader);
+		check(PixelShader.IsValid());
+		GraphicsPSOInit.BoundShaderState.VertexShaderRHI = VertexShader.GetVertexShader();
+		GraphicsPSOInit.BoundShaderState.PixelShaderRHI = PixelShader.GetPixelShader();
 	}
 
 	SetGraphicsPipelineState(RHICmdList, GraphicsPSOInit);
@@ -381,21 +381,21 @@ void FAppleARKitVideoOverlay::RenderVideoOverlayWithMaterial(FRHICommandListImme
 
 	if (bIsMobileRenderer)
 	{
-		FARKitCameraOverlayMobileVS* const VertexShaderPtr = static_cast<FARKitCameraOverlayMobileVS*>(VertexShader);
+		FARKitCameraOverlayMobileVS* const VertexShaderPtr = static_cast<FARKitCameraOverlayMobileVS*>(VertexShader.GetShader());
 		check(VertexShaderPtr != nullptr);
-		SetUniformBufferParameterImmediate(RHICmdList, VertexShaderPtr->GetVertexShader(), VertexShaderPtr->GetUniformBufferParameter<FDrawRectangleParameters>(), Parameters);
+		SetUniformBufferParameterImmediate(RHICmdList, VertexShader.GetVertexShader(), VertexShaderPtr->GetUniformBufferParameter<FDrawRectangleParameters>(), Parameters);
 		VertexShaderPtr->SetParameters(RHICmdList, InView);
-		FARKitCameraOverlayMobilePS* PixelShaderPtr = static_cast<FARKitCameraOverlayMobilePS*>(PixelShader);
+		FARKitCameraOverlayMobilePS* PixelShaderPtr = static_cast<FARKitCameraOverlayMobilePS*>(PixelShader.GetShader());
 		check(PixelShaderPtr != nullptr);
 		PixelShaderPtr->SetParameters(RHICmdList, InView, RenderingOverlayMaterial->GetRenderProxy());
 	}
 	else
 	{
-		FARKitCameraOverlayVS* const VertexShaderPtr = static_cast<FARKitCameraOverlayVS*>(VertexShader);
+		FARKitCameraOverlayVS* const VertexShaderPtr = static_cast<FARKitCameraOverlayVS*>(VertexShader.GetShader());
 		check(VertexShaderPtr != nullptr);
-		SetUniformBufferParameterImmediate(RHICmdList, VertexShaderPtr->GetVertexShader(), VertexShaderPtr->GetUniformBufferParameter<FDrawRectangleParameters>(), Parameters);
+		SetUniformBufferParameterImmediate(RHICmdList, VertexShader.GetVertexShader(), VertexShaderPtr->GetUniformBufferParameter<FDrawRectangleParameters>(), Parameters);
 		VertexShaderPtr->SetParameters(RHICmdList, InView);
-		FARKitCameraOverlayPS* PixelShaderPtr = static_cast<FARKitCameraOverlayPS*>(PixelShader);
+		FARKitCameraOverlayPS* PixelShaderPtr = static_cast<FARKitCameraOverlayPS*>(PixelShader.GetShader());
 		check(PixelShaderPtr != nullptr);
 		PixelShaderPtr->SetParameters(RHICmdList, InView, RenderingOverlayMaterial->GetRenderProxy());
 	}
diff --git a/Engine/Plugins/Runtime/AR/Google/GoogleARCore/Source/GoogleARCoreRendering/Private/GoogleARCorePassthroughCameraRenderer.cpp b/Engine/Plugins/Runtime/AR/Google/GoogleARCore/Source/GoogleARCoreRendering/Private/GoogleARCorePassthroughCameraRenderer.cpp
--- a/Engine/Plugins/Runtime/AR/Google/GoogleARCore/Source/GoogleARCoreRendering/Private/GoogleARCorePassthroughCameraRenderer.cpp
+++ b/Engine/Plugins/Runtime/AR/Google/GoogleARCore/Source/GoogleARCoreRendering/Private/GoogleARCorePassthroughCameraRenderer.cpp
@@ -155,15 +155,15 @@ public:
 
 	static bool ShouldCompilePermutation(const FMaterialShaderPermutationParameters& Parameters)
 	{
-		return Parameters.Material->GetMaterialDomain() == MD_PostProcess && IsMobilePlatform(Parameters.Platform);
+		return Parameters.MaterialParameters.MaterialDomain == MD_PostProcess && IsMobilePlatform(Parameters.Platform);
 	}
 
 	static void ModifyCompilationEnvironment(const FMaterialShaderPermutationParameters& Parameters, FShaderCompilerEnvironment& OutEnvironment)
 	{
-		FMaterialShader::ModifyCompilationEnvironment(Parameters.Platform, OutEnvironment);
+		FMaterialShader::ModifyCompilationEnvironment(Parameters, OutEnvironment);
 		OutEnvironment.SetDefine(TEXT("POST_PROCESS_MATERIAL"), 1);
 		OutEnvironment.SetDefine(TEXT("POST_PROCESS_MATERIAL_MOBILE"), 1);
-		OutEnvironment.SetDefine(TEXT("POST_PROCESS_MATERIAL_BEFORE_TONEMAP"), (Parameters.Material->GetBlendableLocation() != BL_AfterTonemapping) ? 1 : 0);
+		OutEnvironment.SetDefine(TEXT("POST_PROCESS_MATERIAL_BEFORE_TONEMAP"), (Parameters.MaterialParameters.BlendableLocation != BL_AfterTonemapping) ? 1 : 0);
 	}
 };
 
@@ -171,7 +171,7 @@ public:
 class FGoogleARCoreCameraOverlayVS : public FPostProcessMaterialShader
 {
 public:
-	DECLARE_MATERIAL_SHADER(FGoogleARCoreCameraOverlayVS);
+	DECLARE_SHADER_TYPE(FGoogleARCoreCameraOverlayVS, Material);
 
 	static void ModifyCompilationEnvironment(const FMaterialShaderPermutationParameters& Parameters, FShaderCompilerEnvironment& OutEnvironment)
 	{
@@ -186,17 +186,17 @@ public:
 
 	void SetParameters(FRHICommandList& RHICmdList, const FSceneView View)
 	{
-		FRHIVertexShader* ShaderRHI = GetVertexShader();
+		FRHIVertexShader* ShaderRHI = RHICmdList.GetBoundVertexShader();
 		FMaterialShader::SetViewParameters(RHICmdList, ShaderRHI, View, View.ViewUniformBuffer);
 	}
 };
 
-IMPLEMENT_MATERIAL_SHADER(FGoogleARCoreCameraOverlayVS, "/Engine/Private/PostProcessMaterialShaders.usf", "MainVS", SF_Vertex);
+IMPLEMENT_GLOBAL_SHADER(FGoogleARCoreCameraOverlayVS, "/Engine/Private/PostProcessMaterialShaders.usf", "MainVS", SF_Vertex);
 
 class FGoogleARCoreCameraOverlayPS : public FPostProcessMaterialShader
 {
 public:
-	DECLARE_MATERIAL_SHADER(FGoogleARCoreCameraOverlayPS);
+	DECLARE_SHADER_TYPE(FGoogleARCoreCameraOverlayPS, Material);
 
 	static void ModifyCompilationEnvironment(const FMaterialShaderPermutationParameters& Parameters, FShaderCompilerEnvironment& OutEnvironment)
 	{
@@ -211,12 +211,13 @@ public:
 
 	void SetParameters(FRHICommandList& RHICmdList, const FSceneView View, const FMaterialRenderProxy* Material)
 	{
-		FRHIPixelShader* ShaderRHI = GetPixelShader();
+		FRHIPixelShader* ShaderRHI = RHICmdList.GetBoundPixelShader();
+
 		FMaterialShader::SetParameters(RHICmdList, ShaderRHI, Material, *Material->GetMaterial(View.GetFeatureLevel()), View, View.ViewUniformBuffer, ESceneTextureSetupMode::None);
 	}
 };
 
-IMPLEMENT_MATERIAL_SHADER(FGoogleARCoreCameraOverlayPS, "/Engine/Private/PostProcessMaterialShaders.usf", "MainPS", SF_Pixel);
+IMPLEMENT_GLOBAL_SHADER(FGoogleARCoreCameraOverlayPS, "/Engine/Private/PostProcessMaterialShaders.usf", "MainPS", SF_Pixel);
 
 void FGoogleARCorePassthroughCameraRenderer::RenderVideoOverlay_RenderThread(FRHICommandListImmediate& RHICmdList, FSceneView& InView)
 {
@@ -251,8 +252,8 @@ void FGoogleARCorePassthroughCameraRenderer::RenderVideoOverlay_RenderThread(FRH
 		GraphicsPSOInit.DepthStencilState = TStaticDepthStencilState<false, CF_DepthNearOrEqual>::GetRHI();
 
 		GraphicsPSOInit.BoundShaderState.VertexDeclarationRHI = GFilterVertexDeclaration.VertexDeclarationRHI;
-		GraphicsPSOInit.BoundShaderState.VertexShaderRHI = GETSAFERHISHADER_VERTEX(VertexShader);
-		GraphicsPSOInit.BoundShaderState.PixelShaderRHI = GETSAFERHISHADER_PIXEL(PixelShader);
+		GraphicsPSOInit.BoundShaderState.VertexShaderRHI = VertexShader.GetVertexShader();
+		GraphicsPSOInit.BoundShaderState.PixelShaderRHI = PixelShader.GetPixelShader();
 		GraphicsPSOInit.PrimitiveType = PT_TriangleList;
 
 		SetGraphicsPipelineState(RHICmdList, GraphicsPSOInit);
@@ -270,7 +271,7 @@ void FGoogleARCorePassthroughCameraRenderer::RenderVideoOverlay_RenderThread(FRH
 			1.0f / ViewSize.X, 1.0f / ViewSize.Y,
 			1.0f, 1.0f);
 
-		SetUniformBufferParameterImmediate(RHICmdList, VertexShader->GetVertexShader(), VertexShader->GetUniformBufferParameter<FDrawRectangleParameters>(), Parameters);
+		SetUniformBufferParameterImmediate(RHICmdList, VertexShader.GetVertexShader(), VertexShader->GetUniformBufferParameter<FDrawRectangleParameters>(), Parameters);
 
 		if (OverlayVertexBufferRHI.IsValid() && OverlayIndexBufferRHI.IsValid())
 		{
diff --git a/Engine/Plugins/Runtime/AR/Microsoft/HoloLensAR/Source/HoloLensAR/Private/HoloLensCameraImageTexture.cpp b/Engine/Plugins/Runtime/AR/Microsoft/HoloLensAR/Source/HoloLensAR/Private/HoloLensCameraImageTexture.cpp
--- a/Engine/Plugins/Runtime/AR/Microsoft/HoloLensAR/Source/HoloLensAR/Private/HoloLensCameraImageTexture.cpp
+++ b/Engine/Plugins/Runtime/AR/Microsoft/HoloLensAR/Source/HoloLensAR/Private/HoloLensCameraImageTexture.cpp
@@ -117,11 +117,11 @@ public:
 
 	void SetParameters(FRHICommandListImmediate& RHICmdList, const FShaderResourceViewRHIRef& InTextureY, const FShaderResourceViewRHIRef& InTextureUV)
 	{
-		SetSRVParameter(RHICmdList, GetPixelShader(), TextureY, InTextureY);
-		SetSRVParameter(RHICmdList, GetPixelShader(), TextureUV, InTextureUV);
+		SetSRVParameter(RHICmdList, RHICmdList.GetBoundPixelShader(), TextureY, InTextureY);
+		SetSRVParameter(RHICmdList, RHICmdList.GetBoundPixelShader(), TextureUV, InTextureUV);
 
-		SetSamplerParameter(RHICmdList, GetPixelShader(), PointClampedSamplerY, TStaticSamplerState<SF_Point, AM_Clamp, AM_Clamp, AM_Clamp>::GetRHI());
-		SetSamplerParameter(RHICmdList, GetPixelShader(), BilinearClampedSamplerUV, TStaticSamplerState<SF_Bilinear, AM_Clamp, AM_Clamp, AM_Clamp>::GetRHI());
+		SetSamplerParameter(RHICmdList, RHICmdList.GetBoundPixelShader(), PointClampedSamplerY, TStaticSamplerState<SF_Point, AM_Clamp, AM_Clamp, AM_Clamp>::GetRHI());
+		SetSamplerParameter(RHICmdList, RHICmdList.GetBoundPixelShader(), BilinearClampedSamplerUV, TStaticSamplerState<SF_Bilinear, AM_Clamp, AM_Clamp, AM_Clamp>::GetRHI());
 	}
 
 	static bool ShouldCompilePermutation(const FGlobalShaderPermutationParameters& Parameters)
@@ -129,18 +129,11 @@ public:
 		return IsFeatureLevelSupported(Parameters.Platform, ERHIFeatureLevel::ES2);
 	}
 
-	virtual bool Serialize(FArchive& Ar) override
-	{
-		bool bShaderHasOutdatedParameters = FGlobalShader::Serialize(Ar);
-		Ar << TextureY << TextureUV << PointClampedSamplerY << BilinearClampedSamplerUV;
-		return bShaderHasOutdatedParameters;
-	}
-
 private:
-	FShaderResourceParameter TextureY;
-	FShaderResourceParameter TextureUV;
-	FShaderResourceParameter PointClampedSamplerY;
-	FShaderResourceParameter BilinearClampedSamplerUV;
+	LAYOUT_FIELD(FShaderResourceParameter, TextureY);
+	LAYOUT_FIELD(FShaderResourceParameter, TextureUV);
+	LAYOUT_FIELD(FShaderResourceParameter, PointClampedSamplerY);
+	LAYOUT_FIELD(FShaderResourceParameter, BilinearClampedSamplerUV);
 };
 
 IMPLEMENT_SHADER_TYPE(, FHoloLensCameraImageConversionPS, TEXT("/Plugin/HoloLensAR/HoloLensCameraImageConversion.usf"), TEXT("MainPS"), SF_Pixel)
@@ -304,13 +297,13 @@ private:
 			GraphicsPSOInit.BlendState = TStaticBlendStateWriteMask<CW_RGBA, CW_NONE, CW_NONE, CW_NONE, CW_NONE, CW_NONE, CW_NONE, CW_NONE>::GetRHI();
 			GraphicsPSOInit.PrimitiveType = PT_TriangleList;
 
-			TShaderMap<FGlobalShaderType>* GlobalShaderMap = GetGlobalShaderMap(GMaxRHIFeatureLevel);
+			FGlobalShaderMap* GlobalShaderMap = GetGlobalShaderMap(GMaxRHIFeatureLevel);
 			TShaderMapRef<FHoloLensCameraImageConversionVS> VertexShader(GlobalShaderMap);
 			TShaderMapRef<FHoloLensCameraImageConversionPS> PixelShader(GlobalShaderMap);
 
 			GraphicsPSOInit.BoundShaderState.VertexDeclarationRHI = GHoloLensCameraImageConversionVertexDeclaration.VertexDeclarationRHI;
-			GraphicsPSOInit.BoundShaderState.VertexShaderRHI = GETSAFERHISHADER_VERTEX(*VertexShader);
-			GraphicsPSOInit.BoundShaderState.PixelShaderRHI = GETSAFERHISHADER_PIXEL(*PixelShader);
+			GraphicsPSOInit.BoundShaderState.VertexShaderRHI = VertexShader.GetVertexShader();
+			GraphicsPSOInit.BoundShaderState.PixelShaderRHI = PixelShader.GetPixelShader();
 
 			SetGraphicsPipelineState(RHICmdList, GraphicsPSOInit);
 
diff --git a/Engine/Plugins/Runtime/AR/Microsoft/HoloLensAR/Source/HoloLensTargetPlatform/Private/HoloLensTargetPlatform.cpp b/Engine/Plugins/Runtime/AR/Microsoft/HoloLensAR/Source/HoloLensTargetPlatform/Private/HoloLensTargetPlatform.cpp
--- a/Engine/Plugins/Runtime/AR/Microsoft/HoloLensAR/Source/HoloLensTargetPlatform/Private/HoloLensTargetPlatform.cpp
+++ b/Engine/Plugins/Runtime/AR/Microsoft/HoloLensAR/Source/HoloLensTargetPlatform/Private/HoloLensTargetPlatform.cpp
@@ -105,11 +105,6 @@ void FHoloLensTargetPlatform::GetReflectionCaptureFormats(TArray<FName>& OutForm
 	OutFormats.Add(FName(TEXT("EncodedHDR")));
 }
 
-const FPlatformAudioCookOverrides* FHoloLensTargetPlatform::GetAudioCompressionSettings() const
-{
-	return nullptr;
-}
-
 void FHoloLensTargetPlatform::GetTextureFormats(const UTexture* InTexture, TArray< TArray<FName> >& OutFormats) const
 {
 	GetDefaultTextureFormatNamePerLayer(OutFormats.AddDefaulted_GetRef(), this, InTexture, EngineSettings, false);
diff --git a/Engine/Plugins/Runtime/AR/Microsoft/HoloLensAR/Source/HoloLensTargetPlatform/Private/HoloLensTargetPlatform.h b/Engine/Plugins/Runtime/AR/Microsoft/HoloLensAR/Source/HoloLensTargetPlatform/Private/HoloLensTargetPlatform.h
--- a/Engine/Plugins/Runtime/AR/Microsoft/HoloLensAR/Source/HoloLensTargetPlatform/Private/HoloLensTargetPlatform.h
+++ b/Engine/Plugins/Runtime/AR/Microsoft/HoloLensAR/Source/HoloLensTargetPlatform/Private/HoloLensTargetPlatform.h
@@ -68,8 +68,6 @@ public:
 #if WITH_ENGINE
 	virtual void GetReflectionCaptureFormats(TArray<FName>& OutFormats) const override;
 
-	virtual const FPlatformAudioCookOverrides* GetAudioCompressionSettings() const override;
-
 	virtual const class FStaticMeshLODSettings& GetStaticMeshLODSettings() const override { return StaticMeshLODSettings; }
 
 	virtual void GetTextureFormats(const UTexture* InTexture, TArray< TArray<FName> >& OutFormats) const override;
diff --git a/Engine/Plugins/Runtime/CableComponent/Source/CableComponent/Classes/CableComponent.h b/Engine/Plugins/Runtime/CableComponent/Source/CableComponent/Classes/CableComponent.h
--- a/Engine/Plugins/Runtime/CableComponent/Source/CableComponent/Classes/CableComponent.h
+++ b/Engine/Plugins/Runtime/CableComponent/Source/CableComponent/Classes/CableComponent.h
@@ -39,7 +39,7 @@ public:
 	virtual void OnRegister() override;
 	virtual void TickComponent(float DeltaTime, enum ELevelTick TickType, FActorComponentTickFunction *ThisTickFunction) override;
 	virtual void SendRenderDynamicData_Concurrent() override;
-	virtual void CreateRenderState_Concurrent() override;
+	virtual void CreateRenderState_Concurrent(FRegisterComponentContext* Context) override;
 	virtual void ApplyWorldOffset(const FVector& InOffset, bool bWorldShift) override;
 	//~ End UActorComponent Interface.
 
diff --git a/Engine/Plugins/Runtime/CableComponent/Source/CableComponent/Private/CableComponent.cpp b/Engine/Plugins/Runtime/CableComponent/Source/CableComponent/Private/CableComponent.cpp
--- a/Engine/Plugins/Runtime/CableComponent/Source/CableComponent/Private/CableComponent.cpp
+++ b/Engine/Plugins/Runtime/CableComponent/Source/CableComponent/Private/CableComponent.cpp
@@ -674,9 +674,9 @@ void UCableComponent::TickComponent(float DeltaTime, enum ELevelTick TickType, F
 	UpdateComponentToWorld();
 };
 
-void UCableComponent::CreateRenderState_Concurrent()
+void UCableComponent::CreateRenderState_Concurrent(FRegisterComponentContext* Context)
 {
-	Super::CreateRenderState_Concurrent();
+	Super::CreateRenderState_Concurrent(Context);
 
 	SendRenderDynamicData_Concurrent();
 }
diff --git a/Engine/Plugins/Runtime/GoogleVR/GoogleVRHMD/Source/GoogleVRHMD/Private/GoogleVRHMD.cpp b/Engine/Plugins/Runtime/GoogleVR/GoogleVRHMD/Source/GoogleVRHMD/Private/GoogleVRHMD.cpp
--- a/Engine/Plugins/Runtime/GoogleVR/GoogleVRHMD/Source/GoogleVRHMD/Private/GoogleVRHMD.cpp
+++ b/Engine/Plugins/Runtime/GoogleVR/GoogleVRHMD/Source/GoogleVRHMD/Private/GoogleVRHMD.cpp
@@ -1274,8 +1274,8 @@ void FGoogleVRHMD::PostRenderViewFamily_RenderThread(FRHICommandListImmediate& R
 			GraphicsPSOInit.DepthStencilState = TStaticDepthStencilState<false, CF_Always>::GetRHI();
 
 			GraphicsPSOInit.BoundShaderState.VertexDeclarationRHI = GFilterVertexDeclaration.VertexDeclarationRHI;
-			GraphicsPSOInit.BoundShaderState.VertexShaderRHI = GETSAFERHISHADER_VERTEX(*VertexShader);
-			GraphicsPSOInit.BoundShaderState.PixelShaderRHI = GETSAFERHISHADER_PIXEL(*PixelShader);
+			GraphicsPSOInit.BoundShaderState.VertexShaderRHI = VertexShader.GetVertexShader();
+			GraphicsPSOInit.BoundShaderState.PixelShaderRHI = PixelShader.GetPixelShader();
 			GraphicsPSOInit.PrimitiveType = PT_TriangleList;
 
 			SetGraphicsPipelineState(RHICmdList, GraphicsPSOInit);
@@ -1291,7 +1291,7 @@ void FGoogleVRHMD::PostRenderViewFamily_RenderThread(FRHICommandListImmediate& R
 				1, 1,		// Source USize, VSize
 				ReadbackTextureSizes[textureIndex],		// Target buffer size
 				FIntPoint(1, 1),		// Source texture size
-				*VertexShader,
+				VertexShader,
 				EDRF_Default);
 		}
 		RHICmdList.EndRenderPass();
diff --git a/Engine/Plugins/Runtime/GoogleVR/GoogleVRHMD/Source/GoogleVRHMD/Private/GoogleVRHMDRender.cpp b/Engine/Plugins/Runtime/GoogleVR/GoogleVRHMD/Source/GoogleVRHMD/Private/GoogleVRHMDRender.cpp
--- a/Engine/Plugins/Runtime/GoogleVR/GoogleVRHMD/Source/GoogleVRHMD/Private/GoogleVRHMDRender.cpp
+++ b/Engine/Plugins/Runtime/GoogleVR/GoogleVRHMD/Source/GoogleVRHMD/Private/GoogleVRHMDRender.cpp
@@ -198,8 +198,8 @@ static void ResolvePendingRenderTarget(FRHICommandListImmediate& RHICmdList, FGr
 		RHICmdList.ApplyCachedRenderTargets(GraphicsPSOInit);
 
 		GraphicsPSOInit.BoundShaderState.VertexDeclarationRHI = GFilterVertexDeclaration.VertexDeclarationRHI;
-		GraphicsPSOInit.BoundShaderState.VertexShaderRHI = GETSAFERHISHADER_VERTEX(*VertexShader);
-		GraphicsPSOInit.BoundShaderState.PixelShaderRHI = GETSAFERHISHADER_PIXEL(*PixelShader);
+		GraphicsPSOInit.BoundShaderState.VertexShaderRHI = VertexShader.GetVertexShader();
+		GraphicsPSOInit.BoundShaderState.PixelShaderRHI = PixelShader.GetPixelShader();
 		GraphicsPSOInit.PrimitiveType = PT_TriangleList;
 
 		RHICmdList.DrawIndexedPrimitive(
@@ -284,8 +284,8 @@ void FGoogleVRHMD::RenderTexture_RenderThread(FRHICommandListImmediate& RHICmdLi
 			TShaderMapRef<FScreenPS> PixelShader(ShaderMap);
 
 			GraphicsPSOInit.BoundShaderState.VertexDeclarationRHI = GFilterVertexDeclaration.VertexDeclarationRHI;
-			GraphicsPSOInit.BoundShaderState.VertexShaderRHI = GETSAFERHISHADER_VERTEX(*VertexShader);
-			GraphicsPSOInit.BoundShaderState.PixelShaderRHI = GETSAFERHISHADER_PIXEL(*PixelShader);
+			GraphicsPSOInit.BoundShaderState.VertexShaderRHI = VertexShader.GetVertexShader();
+			GraphicsPSOInit.BoundShaderState.PixelShaderRHI = PixelShader.GetPixelShader();
 			GraphicsPSOInit.PrimitiveType = PT_TriangleList;
 
 			SetGraphicsPipelineState(RHICmdList, GraphicsPSOInit);
@@ -300,7 +300,7 @@ void FGoogleVRHMD::RenderTexture_RenderThread(FRHICommandListImmediate& RHICmdLi
 				1.0f, 1.0f,
 				FIntPoint(ViewportWidth, ViewportHeight),
 				FIntPoint(1, 1),
-				*VertexShader,
+				VertexShader,
 				EDRF_Default);
 		}
 		RHICmdList.EndRenderPass();
diff --git a/Engine/Plugins/Runtime/GoogleVR/GoogleVRHMD/Source/GoogleVRHMD/Private/GoogleVRSplash.cpp b/Engine/Plugins/Runtime/GoogleVR/GoogleVRHMD/Source/GoogleVRHMD/Private/GoogleVRSplash.cpp
--- a/Engine/Plugins/Runtime/GoogleVR/GoogleVRHMD/Source/GoogleVRHMD/Private/GoogleVRSplash.cpp
+++ b/Engine/Plugins/Runtime/GoogleVR/GoogleVRHMD/Source/GoogleVRHMD/Private/GoogleVRSplash.cpp
@@ -230,8 +230,8 @@ void FGoogleVRSplash::RenderStereoSplashScreen(FRHICommandListImmediate& RHICmdL
 			GraphicsPSOInit.DepthStencilState = TStaticDepthStencilState<false, CF_Always>::GetRHI();
 
 			GraphicsPSOInit.BoundShaderState.VertexDeclarationRHI = GFilterVertexDeclaration.VertexDeclarationRHI;
-			GraphicsPSOInit.BoundShaderState.VertexShaderRHI = GETSAFERHISHADER_VERTEX(*VertexShader);
-			GraphicsPSOInit.BoundShaderState.PixelShaderRHI = GETSAFERHISHADER_PIXEL(*PixelShader);
+			GraphicsPSOInit.BoundShaderState.VertexShaderRHI = VertexShader.GetVertexShader();
+			GraphicsPSOInit.BoundShaderState.PixelShaderRHI = PixelShader.GetPixelShader();
 			GraphicsPSOInit.PrimitiveType = PT_TriangleList;
 
 			SetGraphicsPipelineState(RHICmdList, GraphicsPSOInit);
@@ -260,7 +260,7 @@ void FGoogleVRSplash::RenderStereoSplashScreen(FRHICommandListImmediate& RHICmdL
 				USize, VSize,
 				FIntPoint(ViewportWidth, ViewportHeight),
 				FIntPoint(1, 1),
-				*VertexShader,
+				VertexShader,
 				EDRF_Default);
 
 			RenderOffsetX = -SplashScreenEyeOffset.X * ViewportWidthPerEye + (1.0f - RenderScale) * ViewportWidthPerEye * 0.5f;
@@ -281,7 +281,7 @@ void FGoogleVRSplash::RenderStereoSplashScreen(FRHICommandListImmediate& RHICmdL
 				USize, VSize,
 				FIntPoint(ViewportWidth, ViewportHeight),
 				FIntPoint(1, 1),
-				*VertexShader,
+				VertexShader,
 				EDRF_Default);
 		}
 	}
diff --git a/Engine/Plugins/Runtime/NetworkPrediction/Source/NetworkPrediction/Public/NetworkedSimulationModelCueTraits.h b/Engine/Plugins/Runtime/NetworkPrediction/Source/NetworkPrediction/Public/NetworkedSimulationModelCueTraits.h
--- a/Engine/Plugins/Runtime/NetworkPrediction/Source/NetworkPrediction/Public/NetworkedSimulationModelCueTraits.h
+++ b/Engine/Plugins/Runtime/NetworkPrediction/Source/NetworkPrediction/Public/NetworkedSimulationModelCueTraits.h
@@ -218,7 +218,7 @@ struct TNetSimCueTraits
 	static constexpr ENetSimCueInvoker InvokeMask { Traits::InvokeMask };
 	static constexpr bool Resimulate { Traits::Resimulate };
 	static constexpr ENetSimCueReplicationTarget ReplicationTarget { Traits::ReplicationTarget };
-	static constexpr ESimulationTickContext SimTickMask { GetSimTickMask( (ENetSimCueInvoker)InvokeMask, Resimulate) };
+	static constexpr ESimulationTickContext SimTickMask() { return GetSimTickMask((ENetSimCueInvoker)InvokeMask, Resimulate); };
 };
 
 // Type requirements: helper to determine if NetSerialize/NetIdentical functions need to be defined for user types based on the above traits
diff --git a/Engine/Plugins/Runtime/NetworkPrediction/Source/NetworkPrediction/Public/NetworkedSimulationModelCues.h b/Engine/Plugins/Runtime/NetworkPrediction/Source/NetworkPrediction/Public/NetworkedSimulationModelCues.h
--- a/Engine/Plugins/Runtime/NetworkPrediction/Source/NetworkPrediction/Public/NetworkedSimulationModelCues.h
+++ b/Engine/Plugins/Runtime/NetworkPrediction/Source/NetworkPrediction/Public/NetworkedSimulationModelCues.h
@@ -365,7 +365,7 @@ struct FNetSimCueDispatcher
 	{
 		if (EnsureValidContext())
 		{
-			if (EnumHasAnyFlags(TNetSimCueTraits<T>::SimTickMask, Context.TickContext))
+			if (EnumHasAnyFlags(TNetSimCueTraits<T>::SimTickMask(), Context.TickContext))
 			{
 				constexpr bool bSupportsResimulate = TNetSimCueTraits<T>::Resimulate;
 
@@ -425,7 +425,7 @@ struct FNetSimCueDispatcher
 			}
 			else
 			{
-				UE_LOG(LogNetSimCues, Log, TEXT("%s .Suppressing Cue Invocation %s. Mask: %d. TickContext: %d"), *GetDebugName(), *FGlobalCueTypeTable::Get().GetTypeName(T::ID), TNetSimCueTraits<T>::SimTickMask, (int32)Context.TickContext);
+				UE_LOG(LogNetSimCues, Log, TEXT("%s .Suppressing Cue Invocation %s. Mask: %d. TickContext: %d"), *GetDebugName(), *FGlobalCueTypeTable::Get().GetTypeName(T::ID), TNetSimCueTraits<T>::SimTickMask(), (int32)Context.TickContext);
 			}
 		}
 	}
diff --git a/Engine/Plugins/Runtime/Oculus/OculusVR/Source/OculusHMD/Private/OculusHMD_CustomPresent.cpp b/Engine/Plugins/Runtime/Oculus/OculusVR/Source/OculusHMD/Private/OculusHMD_CustomPresent.cpp
--- a/Engine/Plugins/Runtime/Oculus/OculusVR/Source/OculusHMD/Private/OculusHMD_CustomPresent.cpp
+++ b/Engine/Plugins/Runtime/Oculus/OculusVR/Source/OculusHMD/Private/OculusHMD_CustomPresent.cpp
@@ -387,7 +387,7 @@ void FCustomPresent::CopyTexture_RenderThread(FRHICommandListImmediate& RHICmdLi
 	auto ShaderMap = GetGlobalShaderMap(FeatureLevel);
 	TShaderMapRef<FScreenVS> VertexShader(ShaderMap);
 	GraphicsPSOInit.BoundShaderState.VertexDeclarationRHI = GFilterVertexDeclaration.VertexDeclarationRHI;
-	GraphicsPSOInit.BoundShaderState.VertexShaderRHI = GETSAFERHISHADER_VERTEX(*VertexShader);
+	GraphicsPSOInit.BoundShaderState.VertexShaderRHI = VertexShader.GetVertexShader();
 
 	if (DstTexture2D)
 	{
@@ -409,14 +409,14 @@ void FCustomPresent::CopyTexture_RenderThread(FRHICommandListImmediate& RHICmdLi
 			if (!sRGBSource)
 			{
 				TShaderMapRef<FScreenPS> PixelShader(ShaderMap);
-				GraphicsPSOInit.BoundShaderState.PixelShaderRHI = GETSAFERHISHADER_PIXEL(*PixelShader);
+				GraphicsPSOInit.BoundShaderState.PixelShaderRHI = PixelShader.GetPixelShader();
 				SetGraphicsPipelineState(RHICmdList, GraphicsPSOInit);
 				PixelShader->SetParameters(RHICmdList, SamplerState, SrcTextureRHI);
 			}
 			else
 			{
 				TShaderMapRef<FScreenPSsRGBSource> PixelShader(ShaderMap);
-				GraphicsPSOInit.BoundShaderState.PixelShaderRHI = GETSAFERHISHADER_PIXEL(*PixelShader);
+				GraphicsPSOInit.BoundShaderState.PixelShaderRHI = PixelShader.GetPixelShader();
 				SetGraphicsPipelineState(RHICmdList, GraphicsPSOInit);
 				PixelShader->SetParameters(RHICmdList, SamplerState, SrcTextureRHI);
 			}
@@ -429,7 +429,7 @@ void FCustomPresent::CopyTexture_RenderThread(FRHICommandListImmediate& RHICmdLi
 				U, V, USize, VSize,
 				TargetSize,
 				FIntPoint(1, 1),
-				*VertexShader,
+				VertexShader,
 				EDRF_Default);
 		}
 		RHICmdList.EndRenderPass();
@@ -451,7 +451,7 @@ void FCustomPresent::CopyTexture_RenderThread(FRHICommandListImmediate& RHICmdLi
 				RHICmdList.ApplyCachedRenderTargets(GraphicsPSOInit);
 
 				TShaderMapRef<FOculusCubemapPS> PixelShader(ShaderMap);
-				GraphicsPSOInit.BoundShaderState.PixelShaderRHI = GETSAFERHISHADER_PIXEL(*PixelShader);
+				GraphicsPSOInit.BoundShaderState.PixelShaderRHI = PixelShader.GetPixelShader();
 				SetGraphicsPipelineState(RHICmdList, GraphicsPSOInit);
 				FRHISamplerState* SamplerState = DstRect.Size() == SrcRect.Size() ? TStaticSamplerState<SF_Point>::GetRHI() : TStaticSamplerState<SF_Bilinear>::GetRHI();
 				PixelShader->SetParameters(RHICmdList, SamplerState, SrcTextureRHI, FaceIndex);
@@ -468,7 +468,7 @@ void FCustomPresent::CopyTexture_RenderThread(FRHICommandListImmediate& RHICmdLi
 #endif
 					TargetSize,
 					FIntPoint(1, 1),
-					*VertexShader,
+					VertexShader,
 					EDRF_Default);
 			}
 			RHICmdList.EndRenderPass();
diff --git a/Engine/Plugins/Runtime/Oculus/OculusVR/Source/OculusHMD/Private/OculusHMD_StressTester.cpp b/Engine/Plugins/Runtime/Oculus/OculusVR/Source/OculusHMD/Private/OculusHMD_StressTester.cpp
--- a/Engine/Plugins/Runtime/Oculus/OculusVR/Source/OculusHMD/Private/OculusHMD_StressTester.cpp
+++ b/Engine/Plugins/Runtime/Oculus/OculusVR/Source/OculusHMD/Private/OculusHMD_StressTester.cpp
@@ -142,7 +142,7 @@ public:
 	//This function is required to let us bind our runtime surface to the shader using an SRV.
 	void SetSurfaces(FRHICommandList& RHICmdList, FShaderResourceViewRHIRef TextureParameterSRV)
 	{
-		FRHIPixelShader* PixelShaderRHI = GetPixelShader();
+		FRHIPixelShader* PixelShaderRHI = RHICmdList.GetBoundPixelShader();
 
 		if (TextureParameter.IsBound())
 		{
@@ -159,14 +159,14 @@ public:
 		ConstantParametersBuffer = FOculusPixelShaderConstantParametersRef::CreateUniformBufferImmediate(ConstantParameters, UniformBuffer_SingleDraw);
 		VariableParametersBuffer = FOculusPixelShaderVariableParametersRef::CreateUniformBufferImmediate(VariableParameters, UniformBuffer_SingleDraw);
 
-		SetUniformBufferParameter(RHICmdList, GetPixelShader(), GetUniformBufferParameter<FOculusPixelShaderConstantParameters>(), ConstantParametersBuffer);
-		SetUniformBufferParameter(RHICmdList, GetPixelShader(), GetUniformBufferParameter<FOculusPixelShaderVariableParameters>(), VariableParametersBuffer);
+		SetUniformBufferParameter(RHICmdList, RHICmdList.GetBoundPixelShader(), GetUniformBufferParameter<FOculusPixelShaderConstantParameters>(), ConstantParametersBuffer);
+		SetUniformBufferParameter(RHICmdList, RHICmdList.GetBoundPixelShader(), GetUniformBufferParameter<FOculusPixelShaderVariableParameters>(), VariableParametersBuffer);
 	}
 
 	//This is used to clean up the buffer binds after each invocation to let them be changed and used elsewhere if needed.
 	void UnbindBuffers(FRHICommandList& RHICmdList)
 	{
-		FRHIPixelShader* PixelShaderRHI = GetPixelShader();
+		FRHIPixelShader* PixelShaderRHI = RHICmdList.GetBoundPixelShader();
 
 		if (TextureParameter.IsBound())
 		{
@@ -176,7 +176,7 @@ public:
 
 private:
 	//This is how you declare resources that are going to be made available in the HLSL
-	FShaderResourceParameter TextureParameter;
+	LAYOUT_FIELD(FShaderResourceParameter, TextureParameter)
 };
 
 
@@ -371,8 +371,8 @@ void FStressTester::DoTickGPU_RenderThread(FRHICommandListImmediate& RHICmdList,
 			TShaderMapRef<FOculusStressShadersPS> PixelShader(GetGlobalShaderMap(FeatureLevel));
 
 			GraphicsPSOInit.BoundShaderState.VertexDeclarationRHI = GOculusTextureVertexDeclaration.VertexDeclarationRHI;
-			GraphicsPSOInit.BoundShaderState.VertexShaderRHI = GETSAFERHISHADER_VERTEX(*VertexShader);
-			GraphicsPSOInit.BoundShaderState.PixelShaderRHI = GETSAFERHISHADER_PIXEL(*PixelShader);
+			GraphicsPSOInit.BoundShaderState.VertexShaderRHI = VertexShader.GetVertexShader();
+			GraphicsPSOInit.BoundShaderState.PixelShaderRHI = PixelShader.GetPixelShader();
 			GraphicsPSOInit.PrimitiveType = PT_TriangleStrip;
 
 			SetGraphicsPipelineState(RHICmdList, GraphicsPSOInit);
diff --git a/Engine/Plugins/Runtime/PacketHandlers/AESGCMHandlerComponent/Source/Private/AESGCMHandlerComponent.cpp b/Engine/Plugins/Runtime/PacketHandlers/AESGCMHandlerComponent/Source/Private/AESGCMHandlerComponent.cpp
--- a/Engine/Plugins/Runtime/PacketHandlers/AESGCMHandlerComponent/Source/Private/AESGCMHandlerComponent.cpp
+++ b/Engine/Plugins/Runtime/PacketHandlers/AESGCMHandlerComponent/Source/Private/AESGCMHandlerComponent.cpp
@@ -4,6 +4,12 @@
 
 IMPLEMENT_MODULE( FAESGCMHandlerComponentModule, AESGCMHandlerComponent )
 
+
+const int32 FAESGCMHandlerComponent::KeySizeInBytes;
+const int32 FAESGCMHandlerComponent::BlockSizeInBytes;
+const int32 FAESGCMHandlerComponent::IVSizeInBytes;
+const int32 FAESGCMHandlerComponent::AuthTagSizeInBytes;
+
 TSharedPtr<HandlerComponent> FAESGCMHandlerComponentModule::CreateComponentInstance(FString& Options)
 {
 	return MakeShared<FAESGCMHandlerComponent>();
diff --git a/Engine/Plugins/Runtime/PacketHandlers/AESHandlerComponent/Source/Private/AESHandlerComponent.cpp b/Engine/Plugins/Runtime/PacketHandlers/AESHandlerComponent/Source/Private/AESHandlerComponent.cpp
--- a/Engine/Plugins/Runtime/PacketHandlers/AESHandlerComponent/Source/Private/AESHandlerComponent.cpp
+++ b/Engine/Plugins/Runtime/PacketHandlers/AESHandlerComponent/Source/Private/AESHandlerComponent.cpp
@@ -4,6 +4,10 @@
 
 IMPLEMENT_MODULE( FAESHandlerComponentModule, AESHandlerComponent )
 
+const int32 FAESHandlerComponent::KeySizeInBytes;
+const int32 FAESHandlerComponent::BlockSizeInBytes;
+
+
 TSharedPtr<HandlerComponent> FAESHandlerComponentModule::CreateComponentInstance(FString& Options)
 {
 	TSharedPtr<HandlerComponent> ReturnVal = NULL;
diff --git a/Engine/Plugins/Runtime/PacketHandlers/DTLSHandlerComponent/Source/Private/DTLSCertificate.cpp b/Engine/Plugins/Runtime/PacketHandlers/DTLSHandlerComponent/Source/Private/DTLSCertificate.cpp
--- a/Engine/Plugins/Runtime/PacketHandlers/DTLSHandlerComponent/Source/Private/DTLSCertificate.cpp
+++ b/Engine/Plugins/Runtime/PacketHandlers/DTLSHandlerComponent/Source/Private/DTLSCertificate.cpp
@@ -9,6 +9,10 @@ THIRD_PARTY_INCLUDES_START
 THIRD_PARTY_INCLUDES_END
 #undef UI
 
+
+constexpr uint32 FDTLSFingerprint::Length;
+
+
 FDTLSCertificate::FDTLSCertificate()
 	: PKey(nullptr)
 	, Certificate(nullptr)
@@ -124,4 +128,4 @@ bool FDTLSCertificate::GenerateCertificate()
 	bSuccess = true;
 
 	return bSuccess;
-}
\ No newline at end of file
+}
diff --git a/Engine/Plugins/Runtime/Steam/SteamVR/Source/SteamVR/Private/SteamVRSpectatorScreen.cpp b/Engine/Plugins/Runtime/Steam/SteamVR/Source/SteamVR/Private/SteamVRSpectatorScreen.cpp
--- a/Engine/Plugins/Runtime/Steam/SteamVR/Source/SteamVR/Private/SteamVRSpectatorScreen.cpp
+++ b/Engine/Plugins/Runtime/Steam/SteamVR/Source/SteamVR/Private/SteamVRSpectatorScreen.cpp
@@ -78,7 +78,7 @@ void FSteamVRHMD::CopyTexture_RenderThread(FRHICommandListImmediate& RHICmdList,
 
 		TShaderMapRef<FScreenVS> VertexShader(ShaderMap);
 		GraphicsPSOInit.BoundShaderState.VertexDeclarationRHI = GFilterVertexDeclaration.VertexDeclarationRHI;
-		GraphicsPSOInit.BoundShaderState.VertexShaderRHI = GETSAFERHISHADER_VERTEX(*VertexShader);
+		GraphicsPSOInit.BoundShaderState.VertexShaderRHI = VertexShader.GetVertexShader();
 
 		const bool bSameSize = DstRect.Size() == SrcRect.Size();
 		FRHISamplerState* PixelSampler = bSameSize ? TStaticSamplerState<SF_Point>::GetRHI() : TStaticSamplerState<SF_Bilinear>::GetRHI();
@@ -86,7 +86,7 @@ void FSteamVRHMD::CopyTexture_RenderThread(FRHICommandListImmediate& RHICmdList,
 		if ((SrcTexture->GetFlags() & TexCreate_SRGB) != 0)
 		{
 			TShaderMapRef<FScreenPSsRGBSource> PixelShader(ShaderMap);
-			GraphicsPSOInit.BoundShaderState.PixelShaderRHI = GETSAFERHISHADER_PIXEL(*PixelShader);
+			GraphicsPSOInit.BoundShaderState.PixelShaderRHI = PixelShader.GetPixelShader();
 
 			SetGraphicsPipelineState(RHICmdList, GraphicsPSOInit);
 			PixelShader->SetParameters(RHICmdList, PixelSampler, SrcTexture);
@@ -94,7 +94,7 @@ void FSteamVRHMD::CopyTexture_RenderThread(FRHICommandListImmediate& RHICmdList,
 		else
 		{
 			TShaderMapRef<FScreenPS> PixelShader(ShaderMap);
-			GraphicsPSOInit.BoundShaderState.PixelShaderRHI = GETSAFERHISHADER_PIXEL(*PixelShader);
+			GraphicsPSOInit.BoundShaderState.PixelShaderRHI = PixelShader.GetPixelShader();
 
 			SetGraphicsPipelineState(RHICmdList, GraphicsPSOInit);
 			PixelShader->SetParameters(RHICmdList, PixelSampler, SrcTexture);
@@ -108,7 +108,7 @@ void FSteamVRHMD::CopyTexture_RenderThread(FRHICommandListImmediate& RHICmdList,
 			USize, VSize,
 			TargetSize,
 			FIntPoint(1, 1),
-			*VertexShader,
+			VertexShader,
 			EDRF_Default);
 	}
 	RHICmdList.EndRenderPass();
diff --git a/Engine/Plugins/Runtime/WarpUtils/Source/PFMExporter/Private/PFMExporterShader.cpp b/Engine/Plugins/Runtime/WarpUtils/Source/PFMExporter/Private/PFMExporterShader.cpp
--- a/Engine/Plugins/Runtime/WarpUtils/Source/PFMExporter/Private/PFMExporterShader.cpp
+++ b/Engine/Plugins/Runtime/WarpUtils/Source/PFMExporter/Private/PFMExporterShader.cpp
@@ -50,17 +50,8 @@ public:
 		SetShaderValue(RHICmdList, ShaderRHI, MeshToPFMMatrixParameter, MeshToPFMMatrix);
 	}
 
-	virtual bool Serialize(FArchive& Ar) override
-	{
-		bool bShaderHasOutdatedParameters = FGlobalShader::Serialize(Ar);
-		Ar
-			<< MeshToPFMMatrixParameter;
-
-		return bShaderHasOutdatedParameters;
-	}
-
 private:
-	FShaderParameter MeshToPFMMatrixParameter;
+	LAYOUT_FIELD(FShaderParameter, MeshToPFMMatrixParameter)
 };
 
 class FPFMExporterPS 
@@ -149,7 +140,7 @@ bool FPFMExporterShader::ApplyPFMExporter_RenderThread(
 			RHICmdList.SetViewport(DstRect.Min.X, DstRect.Min.Y, 0.0f, DstRect.Max.X, DstRect.Max.Y, 1.0f);
 			//DrawClearQuad(RHICmdList, FLinearColor::Black); //!
 
-			TShaderMap<FGlobalShaderType>* ShaderMap = GetGlobalShaderMap(GMaxRHIFeatureLevel);
+			FGlobalShaderMap* ShaderMap = GetGlobalShaderMap(GMaxRHIFeatureLevel);
 			TShaderMapRef<FPFMExporterVS> VertexShader(ShaderMap);
 			TShaderMapRef<FPFMExporterPS> PixelShader(ShaderMap);
 
@@ -162,13 +153,13 @@ bool FPFMExporterShader::ApplyPFMExporter_RenderThread(
 				GraphicsPSOInit.RasterizerState = TStaticRasterizerState<>::GetRHI();
 
 				GraphicsPSOInit.BoundShaderState.VertexDeclarationRHI = GFilterVertexDeclaration.VertexDeclarationRHI;// GetVertexDeclarationFVector4();
-				GraphicsPSOInit.BoundShaderState.VertexShaderRHI = GETSAFERHISHADER_VERTEX(*VertexShader);
-				GraphicsPSOInit.BoundShaderState.PixelShaderRHI = GETSAFERHISHADER_PIXEL(*PixelShader);
+				GraphicsPSOInit.BoundShaderState.VertexShaderRHI = VertexShader.GetVertexShader();
+				GraphicsPSOInit.BoundShaderState.PixelShaderRHI = PixelShader.GetPixelShader();
 				GraphicsPSOInit.PrimitiveType = PT_TriangleList;
 
 				SetGraphicsPipelineState(RHICmdList, GraphicsPSOInit);
 			}
-			VertexShader->SetMeshToPFMMatrix(RHICmdList, VertexShader->GetVertexShader(), MeshToPFMMatrix);
+			VertexShader->SetMeshToPFMMatrix(RHICmdList, VertexShader.GetVertexShader(), MeshToPFMMatrix);
 
 			// Render mesh to PFM texture:
 			RHICmdList.SetStreamSource(0, VertexBufferRHI, 0);
diff --git a/Engine/Plugins/Runtime/WindowsMixedReality/Source/WindowsMixedRealityHMD/Private/WindowsMixedRealityHMD.cpp b/Engine/Plugins/Runtime/WindowsMixedReality/Source/WindowsMixedRealityHMD/Private/WindowsMixedRealityHMD.cpp
--- a/Engine/Plugins/Runtime/WindowsMixedReality/Source/WindowsMixedRealityHMD/Private/WindowsMixedRealityHMD.cpp
+++ b/Engine/Plugins/Runtime/WindowsMixedReality/Source/WindowsMixedRealityHMD/Private/WindowsMixedRealityHMD.cpp
@@ -69,27 +69,16 @@ public:
 
 	void SetParameters(FRHICommandList& RHICmdList, FRHITexture* DepthTexture)
 	{
-		FRHIPixelShader* PixelShaderRHI = GetPixelShader();
+		FRHIPixelShader* PixelShaderRHI = RHICmdList.GetBoundPixelShader();
 
 		FRHISamplerState* SamplerStateRHI = TStaticSamplerState<SF_Point>::GetRHI();
 		SetTextureParameter(RHICmdList, PixelShaderRHI, InDepthTexture, InTextureSampler, SamplerStateRHI, DepthTexture);
 	}
 
-	virtual bool Serialize(FArchive& Ar) override
-	{
-		bool bShaderHasOutdatedParameters = FGlobalShader::Serialize(Ar);
-
-		// Serialize shader inputs.
-		Ar << InDepthTexture;
-		Ar << InTextureSampler;
-
-		return bShaderHasOutdatedParameters;
-	}
-
 private:
 	// Shader parameters.
-	FShaderResourceParameter InDepthTexture;
-	FShaderResourceParameter InTextureSampler;
+	LAYOUT_FIELD(FShaderResourceParameter, InDepthTexture);
+	LAYOUT_FIELD(FShaderResourceParameter, InTextureSampler);
 };
 
 IMPLEMENT_SHADER_TYPE(, FDepthConversionPS, TEXT("/Plugin/WindowsMixedReality/Private/DepthConversion.usf"), TEXT("MainPixelShader"), SF_Pixel)
@@ -689,7 +678,7 @@ namespace WindowsMixedReality
 
 			TShaderMapRef<FScreenVS> VertexShader(ShaderMap);
 			GraphicsPSOInit.BoundShaderState.VertexDeclarationRHI = GFilterVertexDeclaration.VertexDeclarationRHI;
-			GraphicsPSOInit.BoundShaderState.VertexShaderRHI = GETSAFERHISHADER_VERTEX(*VertexShader);
+			GraphicsPSOInit.BoundShaderState.VertexShaderRHI = VertexShader.GetVertexShader();
 
 			const bool bSameSize = DstRect.Size() == SrcRect.Size();
 			FRHISamplerState* PixelSampler = bSameSize ? TStaticSamplerState<SF_Point>::GetRHI() : TStaticSamplerState<SF_Bilinear>::GetRHI();
@@ -697,7 +686,7 @@ namespace WindowsMixedReality
 			if ((SrcTexture->GetFlags() & TexCreate_SRGB) != 0)
 			{
 				TShaderMapRef<FScreenPSsRGBSource> PixelShader(ShaderMap);
-				GraphicsPSOInit.BoundShaderState.PixelShaderRHI = GETSAFERHISHADER_PIXEL(*PixelShader);
+				GraphicsPSOInit.BoundShaderState.PixelShaderRHI = PixelShader.GetPixelShader();
 
 				SetGraphicsPipelineState(RHICmdList, GraphicsPSOInit);
 				PixelShader->SetParameters(RHICmdList, PixelSampler, SrcTexture);
@@ -705,7 +694,7 @@ namespace WindowsMixedReality
 			else
 			{
 				TShaderMapRef<FScreenPS> PixelShader(ShaderMap);
-				GraphicsPSOInit.BoundShaderState.PixelShaderRHI = GETSAFERHISHADER_PIXEL(*PixelShader);
+				GraphicsPSOInit.BoundShaderState.PixelShaderRHI = PixelShader.GetPixelShader();
 
 				SetGraphicsPipelineState(RHICmdList, GraphicsPSOInit);
 				PixelShader->SetParameters(RHICmdList, PixelSampler, SrcTexture);
@@ -719,7 +708,7 @@ namespace WindowsMixedReality
 				USize, VSize,
 				TargetSize,
 				FIntPoint(1, 1),
-				*VertexShader,
+				VertexShader,
 				EDRF_Default);
 
 			RHICmdList.EndRenderPass();
@@ -1364,8 +1353,8 @@ namespace WindowsMixedReality
 			TShaderMapRef<FDepthConversionPS> pixelShader(shaderMap);
 
 			GraphicsPSOInit.BoundShaderState.VertexDeclarationRHI = GFilterVertexDeclaration.VertexDeclarationRHI;
-			GraphicsPSOInit.BoundShaderState.VertexShaderRHI = GETSAFERHISHADER_VERTEX(*vertexShader);
-			GraphicsPSOInit.BoundShaderState.PixelShaderRHI = GETSAFERHISHADER_PIXEL(*pixelShader);
+			GraphicsPSOInit.BoundShaderState.VertexShaderRHI = vertexShader.GetVertexShader();
+			GraphicsPSOInit.BoundShaderState.PixelShaderRHI = pixelShader.GetPixelShader();
 			GraphicsPSOInit.PrimitiveType = PT_TriangleList;
 
 			SetGraphicsPipelineState(RHICmdList, GraphicsPSOInit);
@@ -1380,7 +1369,7 @@ namespace WindowsMixedReality
 				1.0f, 1.0f, // SizeU, SizeV
 				FIntPoint(viewportWidth, viewportHeight), // TargetSize
 				FIntPoint(1, 1), // TextureSize
-				*vertexShader,
+				vertexShader,
 				EDRF_Default);
 		}
 		RHICmdList.EndRenderPass();
diff --git a/Engine/Plugins/Runtime/WindowsMoviePlayer/Source/WindowsMoviePlayer/Private/WindowsMovieStreamer.cpp b/Engine/Plugins/Runtime/WindowsMoviePlayer/Source/WindowsMoviePlayer/Private/WindowsMovieStreamer.cpp
--- a/Engine/Plugins/Runtime/WindowsMoviePlayer/Source/WindowsMoviePlayer/Private/WindowsMovieStreamer.cpp
+++ b/Engine/Plugins/Runtime/WindowsMoviePlayer/Source/WindowsMoviePlayer/Private/WindowsMovieStreamer.cpp
@@ -136,14 +136,14 @@ void FMediaFoundationMovieStreamer::ConvertSample()
 		TShaderMapRef<FMediaShadersVS> VertexShader(ShaderMap);
 
 		GraphicsPSOInit.BoundShaderState.VertexDeclarationRHI = GMediaVertexDeclaration.VertexDeclarationRHI;
-		GraphicsPSOInit.BoundShaderState.VertexShaderRHI = GETSAFERHISHADER_VERTEX(*VertexShader);
+		GraphicsPSOInit.BoundShaderState.VertexShaderRHI = VertexShader.GetVertexShader();
 
 		switch (SourceFormat.SampleFormat)
 		{
 		case EMediaTextureSampleFormat::CharBMP:
 		{
 			TShaderMapRef<FBMPConvertPS> ConvertShader(ShaderMap);
-			GraphicsPSOInit.BoundShaderState.PixelShaderRHI = GETSAFERHISHADER_PIXEL(*ConvertShader);
+			GraphicsPSOInit.BoundShaderState.PixelShaderRHI = ConvertShader.GetPixelShader();
 			SetGraphicsPipelineState(CommandList, GraphicsPSOInit);
 			ConvertShader->SetParameters(CommandList, InputTarget, OutputDim, bSampleIsOutputSrgb && !SrgbOutput);
 		}
@@ -152,7 +152,7 @@ void FMediaFoundationMovieStreamer::ConvertSample()
 		case EMediaTextureSampleFormat::CharYUY2:
 		{
 			TShaderMapRef<FYUY2ConvertPS> ConvertShader(ShaderMap);
-			GraphicsPSOInit.BoundShaderState.PixelShaderRHI = GETSAFERHISHADER_PIXEL(*ConvertShader);
+			GraphicsPSOInit.BoundShaderState.PixelShaderRHI = ConvertShader.GetPixelShader();
 			SetGraphicsPipelineState(CommandList, GraphicsPSOInit);
 			ConvertShader->SetParameters(CommandList, InputTarget, OutputDim, MediaShaders::YuvToSrgbDefault, MediaShaders::YUVOffset8bits, bSampleIsOutputSrgb);
 		}
diff --git a/Engine/Plugins/Runtime/nDisplay/Source/MPCDI/Private/MPCDIShader.cpp b/Engine/Plugins/Runtime/nDisplay/Source/MPCDI/Private/MPCDIShader.cpp
--- a/Engine/Plugins/Runtime/nDisplay/Source/MPCDI/Private/MPCDIShader.cpp
+++ b/Engine/Plugins/Runtime/nDisplay/Source/MPCDI/Private/MPCDIShader.cpp
@@ -335,7 +335,7 @@ public:
 		}
 
 		// Get mutable shaders:
-		TShaderMap<FGlobalShaderType>* ShaderMap = GetGlobalShaderMap(GMaxRHIFeatureLevel);
+		FGlobalShaderMap* ShaderMap = GetGlobalShaderMap(GMaxRHIFeatureLevel);
 		TShaderMapRef<FMpcdiWarpVS> VertexShader(ShaderMap, RenderPassData.VSPermutationVector);
 		TShaderMapRef<FMpcdiPassthroughPS> PixelShader(ShaderMap);
 
@@ -351,8 +351,8 @@ public:
 			GraphicsPSOInit.DepthStencilState = TStaticDepthStencilState<false, CF_Never>::GetRHI();
 			GraphicsPSOInit.RasterizerState = TStaticRasterizerState<>::GetRHI();
 
-			GraphicsPSOInit.BoundShaderState.VertexShaderRHI = GETSAFERHISHADER_VERTEX(*VertexShader);
-			GraphicsPSOInit.BoundShaderState.PixelShaderRHI = GETSAFERHISHADER_PIXEL(*PixelShader);
+			GraphicsPSOInit.BoundShaderState.VertexShaderRHI = VertexShader.GetVertexShader();
+			GraphicsPSOInit.BoundShaderState.PixelShaderRHI = PixelShader.GetPixelShader();
 
 			GraphicsPSOInit.PrimitiveType = PT_TriangleList;
 			GraphicsPSOInit.BlendState = TStaticBlendState <>::GetRHI();
@@ -360,8 +360,8 @@ public:
 		}
 
 		// Setup shaders data
-		SetShaderParameters(RHICmdList, *VertexShader, VertexShader->GetVertexShader(), RenderPassData.VSParameters);
-		SetShaderParameters(RHICmdList, *PixelShader, PixelShader->GetPixelShader(), RenderPassData.PSParameters);
+		SetShaderParameters(RHICmdList, VertexShader, VertexShader.GetVertexShader(), RenderPassData.VSParameters);
+		SetShaderParameters(RHICmdList, PixelShader, PixelShader.GetPixelShader(), RenderPassData.PSParameters);
 
 		MPCDIRegion->WarpData->FinishRender(RHICmdList);
 
@@ -376,7 +376,7 @@ public:
 		}
 
 		// Get mutable shaders:
-		TShaderMap<FGlobalShaderType>* ShaderMap = GetGlobalShaderMap(GMaxRHIFeatureLevel);
+		FGlobalShaderMap* ShaderMap = GetGlobalShaderMap(GMaxRHIFeatureLevel);
 		TShaderMapRef<FMpcdiWarpVS> VertexShader(ShaderMap, RenderPassData.VSPermutationVector);
 		TShaderMapRef<FMpcdiWarpPS> PixelShader(ShaderMap, RenderPassData.PSPermutationVector);
 
@@ -392,8 +392,8 @@ public:
 			GraphicsPSOInit.DepthStencilState = TStaticDepthStencilState<false, CF_Never>::GetRHI();
 			GraphicsPSOInit.RasterizerState = TStaticRasterizerState<>::GetRHI();
 
-			GraphicsPSOInit.BoundShaderState.VertexShaderRHI = GETSAFERHISHADER_VERTEX(*VertexShader);
-			GraphicsPSOInit.BoundShaderState.PixelShaderRHI = GETSAFERHISHADER_PIXEL(*PixelShader);
+			GraphicsPSOInit.BoundShaderState.VertexShaderRHI = VertexShader.GetVertexShader();
+			GraphicsPSOInit.BoundShaderState.PixelShaderRHI = PixelShader.GetPixelShader();
 
 			GraphicsPSOInit.PrimitiveType = PT_TriangleList;
 
@@ -404,8 +404,8 @@ public:
 		}
 
 		// Setup shaders data
-		SetShaderParameters(RHICmdList, *VertexShader, VertexShader->GetVertexShader(), RenderPassData.VSParameters);
-		SetShaderParameters(RHICmdList, *PixelShader, PixelShader->GetPixelShader(), RenderPassData.PSParameters);
+		SetShaderParameters(RHICmdList, VertexShader, VertexShader.GetVertexShader(), RenderPassData.VSParameters);
+		SetShaderParameters(RHICmdList, PixelShader, PixelShader.GetPixelShader(), RenderPassData.PSParameters);
 
 		MPCDIRegion->WarpData->FinishRender(RHICmdList);
 
diff --git a/Engine/Plugins/Runtime/nDisplay/Source/OutputRemap/Private/OutputRemapShader.cpp b/Engine/Plugins/Runtime/nDisplay/Source/OutputRemap/Private/OutputRemapShader.cpp
--- a/Engine/Plugins/Runtime/nDisplay/Source/OutputRemap/Private/OutputRemapShader.cpp
+++ b/Engine/Plugins/Runtime/nDisplay/Source/OutputRemap/Private/OutputRemapShader.cpp
@@ -42,7 +42,7 @@ class FOutputRemapVS : public FGlobalShader
 
 public:
 	/** Default constructor. */
-	FOutputRemapVS()
+	FOutputRemapVS() 
 	{ }
 
 public:
@@ -90,20 +90,9 @@ public:
 		RHICmdList.SetShaderSampler(ShaderRHI, PostprocessInputParameterSampler0.GetBaseIndex(), TStaticSamplerState<SF_Trilinear, AM_Clamp, AM_Clamp, AM_Clamp>::GetRHI());
 	}
 
-	virtual bool Serialize(FArchive& Ar) override
-	{
-		bool bShaderHasOutdatedParameters = FGlobalShader::Serialize(Ar);
-		Ar
-			<< PostprocessInputParameter0
-			<< PostprocessInputParameterSampler0
-			;
-
-		return bShaderHasOutdatedParameters;
-	}
-
 private:
-	FShaderResourceParameter PostprocessInputParameter0;
-	FShaderResourceParameter PostprocessInputParameterSampler0;
+	LAYOUT_FIELD(FShaderResourceParameter, PostprocessInputParameter0)
+	LAYOUT_FIELD(FShaderResourceParameter, PostprocessInputParameterSampler0)
 };
 
 // Implement shaders inside UE4
@@ -163,19 +152,18 @@ bool FOutputRemapShader::ApplyOutputRemap_RenderThread(FRHICommandListImmediate&
 		GraphicsPSOInit.BlendState = TStaticBlendState <>::GetRHI();
 		GraphicsPSOInit.RasterizerState = TStaticRasterizerState<>::GetRHI();
 
-		// Grab shaders
-		TShaderMap<FGlobalShaderType>* ShaderMap = GetGlobalShaderMap(GMaxRHIFeatureLevel);
+		FGlobalShaderMap* ShaderMap = GetGlobalShaderMap(GMaxRHIFeatureLevel);
 		TShaderMapRef<FOutputRemapVS> VertexShader(ShaderMap);
 		TShaderMapRef<FOutputRemapPS> PixelShader(ShaderMap);
 
-		GraphicsPSOInit.BoundShaderState.VertexDeclarationRHI = GFilterVertexDeclaration.VertexDeclarationRHI;
-		GraphicsPSOInit.BoundShaderState.VertexShaderRHI = GETSAFERHISHADER_VERTEX(*VertexShader);
-		GraphicsPSOInit.BoundShaderState.PixelShaderRHI = GETSAFERHISHADER_PIXEL(*PixelShader);
+		GraphicsPSOInit.BoundShaderState.VertexDeclarationRHI = GFilterVertexDeclaration.VertexDeclarationRHI;// GetVertexDeclarationFVector4();
+		GraphicsPSOInit.BoundShaderState.VertexShaderRHI = VertexShader.GetVertexShader();
+		GraphicsPSOInit.BoundShaderState.PixelShaderRHI = PixelShader.GetPixelShader();
 		GraphicsPSOInit.PrimitiveType = PT_TriangleList;
 
 		SetGraphicsPipelineState(RHICmdList, GraphicsPSOInit);
 
-		PixelShader->SetParameters(RHICmdList, PixelShader->GetPixelShader(), ShaderResourceTexture);
+		PixelShader->SetParameters(RHICmdList, PixelShader.GetPixelShader(), ShaderResourceTexture);
 		MeshData->DrawMesh(RHICmdList);
 	}
 
diff --git a/Engine/Plugins/Runtime/nDisplay/Source/PicpMPCDI/Private/PicpBlurPostProcess.cpp b/Engine/Plugins/Runtime/nDisplay/Source/PicpMPCDI/Private/PicpBlurPostProcess.cpp
--- a/Engine/Plugins/Runtime/nDisplay/Source/PicpMPCDI/Private/PicpBlurPostProcess.cpp
+++ b/Engine/Plugins/Runtime/nDisplay/Source/PicpMPCDI/Private/PicpBlurPostProcess.cpp
@@ -65,8 +65,8 @@ void FPicpBlurPostProcess::ApplyCompose_RenderThread(
 			GraphicsPSOInit.RasterizerState = TStaticRasterizerState<>::GetRHI();
 			GraphicsPSOInit.DepthStencilState = TStaticDepthStencilState<false, CF_Always>::GetRHI();
 			GraphicsPSOInit.BoundShaderState.VertexDeclarationRHI = GFilterVertexDeclaration.VertexDeclarationRHI;
-			GraphicsPSOInit.BoundShaderState.VertexShaderRHI = GETSAFERHISHADER_VERTEX(*VertexShader);
-			GraphicsPSOInit.BoundShaderState.PixelShaderRHI = GETSAFERHISHADER_PIXEL(*PixelShader);
+			GraphicsPSOInit.BoundShaderState.VertexShaderRHI = VertexShader.GetVertexShader();
+			GraphicsPSOInit.BoundShaderState.PixelShaderRHI = PixelShader.GetPixelShader();
 			GraphicsPSOInit.PrimitiveType = PT_TriangleList;
 
 			SetGraphicsPipelineState(RHICmdList, GraphicsPSOInit);
@@ -107,8 +107,8 @@ static void PicpBlurPostProcess_RenderThread(
 	GraphicsPSOInit.RasterizerState = TStaticRasterizerState<>::GetRHI();
 	GraphicsPSOInit.DepthStencilState = TStaticDepthStencilState<false, CF_Always>::GetRHI();
 	GraphicsPSOInit.BoundShaderState.VertexDeclarationRHI = GFilterVertexDeclaration.VertexDeclarationRHI;
-	GraphicsPSOInit.BoundShaderState.VertexShaderRHI = GETSAFERHISHADER_VERTEX(*VertexShader);
-	GraphicsPSOInit.BoundShaderState.PixelShaderRHI = GETSAFERHISHADER_PIXEL(*PixelShader);
+	GraphicsPSOInit.BoundShaderState.VertexShaderRHI = VertexShader.GetVertexShader();
+	GraphicsPSOInit.BoundShaderState.PixelShaderRHI = PixelShader.GetPixelShader();
 	GraphicsPSOInit.PrimitiveType = PT_TriangleList;
 	{
 		const FIntPoint TargetSizeXY = InOutTexture->GetSizeXY();
diff --git a/Engine/Plugins/Runtime/nDisplay/Source/PicpMPCDI/Private/PicpBlurPostProcess.h b/Engine/Plugins/Runtime/nDisplay/Source/PicpMPCDI/Private/PicpBlurPostProcess.h
--- a/Engine/Plugins/Runtime/nDisplay/Source/PicpMPCDI/Private/PicpBlurPostProcess.h
+++ b/Engine/Plugins/Runtime/nDisplay/Source/PicpMPCDI/Private/PicpBlurPostProcess.h
@@ -75,11 +75,10 @@ class TPicpBlurPostProcessPS : public FGlobalShader
 	{
 	}
 
-	FShaderResourceParameter SrcTextureParameter;
-	FShaderResourceParameter BilinearClampTextureSamplerParameter;
-	FShaderParameter SampleOffsetParameter;
-	FShaderParameter KernelRadiusParameter;
-
+	LAYOUT_FIELD(FShaderResourceParameter, SrcTextureParameter);
+	LAYOUT_FIELD(FShaderResourceParameter, BilinearClampTextureSamplerParameter);
+	LAYOUT_FIELD(FShaderParameter, SampleOffsetParameter);
+	LAYOUT_FIELD(FShaderParameter, KernelRadiusParameter);
 
 public:
 
@@ -114,26 +113,13 @@ public:
 
 	void SetParameters(FRHICommandListImmediate& RHICmdList, FRHITexture2D* SourceTexture, FVector2D SampleOffset, int KernelRadius)
 	{
-		FRHIPixelShader* ShaderRHI = GetPixelShader();
+		FRHIPixelShader* ShaderRHI = RHICmdList.GetBoundPixelShader();
 
 		SetSamplerParameter(RHICmdList, ShaderRHI, BilinearClampTextureSamplerParameter, TStaticSamplerState<SF_Bilinear, AM_Clamp, AM_Clamp, AM_Clamp>::GetRHI());
 		SetTextureParameter(RHICmdList, ShaderRHI, SrcTextureParameter, SourceTexture);
 		SetShaderValue(RHICmdList, ShaderRHI, SampleOffsetParameter, SampleOffset);
 		SetShaderValue(RHICmdList, ShaderRHI, KernelRadiusParameter, KernelRadius);
 	}
-
-	// FShader interface.
-	virtual bool Serialize(FArchive& Ar) override
-	{
-		bool bShaderHasOutdatedParameters = FGlobalShader::Serialize(Ar);
-
-		Ar << SrcTextureParameter;
-		Ar << BilinearClampTextureSamplerParameter;
-		Ar << SampleOffsetParameter;
-		Ar << KernelRadiusParameter;
-
-		return bShaderHasOutdatedParameters;
-	}
 };
 typedef TPicpBlurPostProcessPS<(int)EPicpBlurPostProcessShaderType::Gaussian>   FPicpBlurPostProcessDefaultPS;
 typedef TPicpBlurPostProcessPS<(int)EPicpBlurPostProcessShaderType::Dilate>     FPicpBlurPostProcessDilatePS;
@@ -154,8 +140,8 @@ class FDirectComposePS : public FGlobalShader
 	{
 	}
 
-	FShaderResourceParameter SrcTextureParameter;
-	FShaderResourceParameter BilinearClampTextureSamplerParameter;
+	LAYOUT_FIELD(FShaderResourceParameter, SrcTextureParameter);
+	LAYOUT_FIELD(FShaderResourceParameter, BilinearClampTextureSamplerParameter);
 
 public:
 
@@ -169,20 +155,9 @@ public:
 
 	void SetParameters(FRHICommandListImmediate& RHICmdList, FTextureRHIRef SourceTexture)
 	{
-		FRHIPixelShader* ShaderRHI = GetPixelShader();
+		FRHIPixelShader* ShaderRHI = RHICmdList.GetBoundPixelShader();
 
 		SetSamplerParameter(RHICmdList, ShaderRHI, BilinearClampTextureSamplerParameter, TStaticSamplerState<SF_Bilinear, AM_Clamp, AM_Clamp, AM_Clamp>::GetRHI());
 		SetTextureParameter(RHICmdList, ShaderRHI, SrcTextureParameter, SourceTexture);
 	}
-
-	// FShader interface.
-	virtual bool Serialize(FArchive& Ar) override
-	{
-		bool bShaderHasOutdatedParameters = FGlobalShader::Serialize(Ar);
-
-		Ar << SrcTextureParameter;
-		Ar << BilinearClampTextureSamplerParameter;
-
-		return bShaderHasOutdatedParameters;
-	}
 };
diff --git a/Engine/Plugins/Runtime/nDisplay/Source/PicpMPCDI/Private/PicpMPCDIShader.cpp b/Engine/Plugins/Runtime/nDisplay/Source/PicpMPCDI/Private/PicpMPCDIShader.cpp
--- a/Engine/Plugins/Runtime/nDisplay/Source/PicpMPCDI/Private/PicpMPCDIShader.cpp
+++ b/Engine/Plugins/Runtime/nDisplay/Source/PicpMPCDI/Private/PicpMPCDIShader.cpp
@@ -592,7 +592,7 @@ public:
 		}
 
 		// Get mutable shaders:
-		TShaderMap<FGlobalShaderType>* ShaderMap = GetGlobalShaderMap(GMaxRHIFeatureLevel);
+		FGlobalShaderMap* ShaderMap = GetGlobalShaderMap(GMaxRHIFeatureLevel);
 		TShaderMapRef<FPicpWarpVS> VertexShader(ShaderMap, RenderPassData.VSPermutationVector);
 		TShaderMapRef<FPicpPassthroughPS> PixelShader(ShaderMap);
 
@@ -608,8 +608,8 @@ public:
 			GraphicsPSOInit.DepthStencilState = TStaticDepthStencilState<false, CF_Never>::GetRHI();
 			GraphicsPSOInit.RasterizerState = TStaticRasterizerState<>::GetRHI();
 
-			GraphicsPSOInit.BoundShaderState.VertexShaderRHI = GETSAFERHISHADER_VERTEX(*VertexShader);
-			GraphicsPSOInit.BoundShaderState.PixelShaderRHI = GETSAFERHISHADER_PIXEL(*PixelShader);
+			GraphicsPSOInit.BoundShaderState.VertexShaderRHI = VertexShader.GetVertexShader();
+			GraphicsPSOInit.BoundShaderState.PixelShaderRHI = PixelShader.GetPixelShader();
 
 			GraphicsPSOInit.PrimitiveType = PT_TriangleList;
 			GraphicsPSOInit.BlendState = TStaticBlendState <>::GetRHI();
@@ -617,8 +617,8 @@ public:
 		}
 
 		// Setup shaders data
-		SetShaderParameters(RHICmdList, *VertexShader, VertexShader->GetVertexShader(), RenderPassData.VSParameters);
-		SetShaderParameters(RHICmdList, *PixelShader, PixelShader->GetPixelShader(), RenderPassData.PSParameters);
+		SetShaderParameters(RHICmdList, VertexShader, VertexShader.GetVertexShader(), RenderPassData.VSParameters);
+		SetShaderParameters(RHICmdList, PixelShader, PixelShader.GetPixelShader(), RenderPassData.PSParameters);
 
 		MPCDIRegion->WarpData->FinishRender(RHICmdList);
 		return true;
@@ -632,7 +632,7 @@ public:
 		}
 		
 		// Get mutable shaders:
-		TShaderMap<FGlobalShaderType>* ShaderMap = GetGlobalShaderMap(GMaxRHIFeatureLevel);
+		FGlobalShaderMap* ShaderMap = GetGlobalShaderMap(GMaxRHIFeatureLevel);
 		TShaderMapRef<FPicpWarpVS> VertexShader(ShaderMap, RenderPassData.VSPermutationVector);
 		TShaderMapRef<FPicpWarpPS> PixelShader(ShaderMap, RenderPassData.PSPermutationVector);
 
@@ -648,8 +648,8 @@ public:
 			GraphicsPSOInit.DepthStencilState = TStaticDepthStencilState<false, CF_Never>::GetRHI();
 			GraphicsPSOInit.RasterizerState = TStaticRasterizerState<>::GetRHI();
 
-			GraphicsPSOInit.BoundShaderState.VertexShaderRHI = GETSAFERHISHADER_VERTEX(*VertexShader);
-			GraphicsPSOInit.BoundShaderState.PixelShaderRHI  = GETSAFERHISHADER_PIXEL(*PixelShader);
+			GraphicsPSOInit.BoundShaderState.VertexShaderRHI = VertexShader.GetVertexShader();
+			GraphicsPSOInit.BoundShaderState.PixelShaderRHI  = PixelShader.GetPixelShader();
 
 			GraphicsPSOInit.PrimitiveType = PT_TriangleList;
 
@@ -676,8 +676,8 @@ public:
 		}
 
 		// Setup shaders data
-		SetShaderParameters(RHICmdList, *VertexShader, VertexShader->GetVertexShader(), RenderPassData.VSParameters);
-		SetShaderParameters(RHICmdList, *PixelShader, PixelShader->GetPixelShader(), RenderPassData.PSParameters);
+		SetShaderParameters(RHICmdList, VertexShader, VertexShader.GetVertexShader(), RenderPassData.VSParameters);
+		SetShaderParameters(RHICmdList, PixelShader, PixelShader.GetPixelShader(), RenderPassData.PSParameters);
 
 		MPCDIRegion->WarpData->FinishRender(RHICmdList);
 
diff --git a/Engine/Shaders/Public/ShaderVersion.ush b/Engine/Shaders/Public/ShaderVersion.ush
--- a/Engine/Shaders/Public/ShaderVersion.ush
+++ b/Engine/Shaders/Public/ShaderVersion.ush
@@ -3,6 +3,6 @@
 // in Platform.ush (which should be included in any shader) it allows to invalidate the shader DDC.
 // 
 // If you are merging streams and there is a conflict with this GUID you should make a new GUID rather than taking one or the other.
-// GUID = 9ca17c83-a9e4-439e-95f1-750785bd53ab
+// GUID = 46EE2461-A4C7-4DC0-B30B-E763362FF2F1
 
 
diff --git a/Engine/Source/Developer/AllDesktopTargetPlatform/Private/AllDesktopTargetPlatform.cpp b/Engine/Source/Developer/AllDesktopTargetPlatform/Private/AllDesktopTargetPlatform.cpp
--- a/Engine/Source/Developer/AllDesktopTargetPlatform/Private/AllDesktopTargetPlatform.cpp
+++ b/Engine/Source/Developer/AllDesktopTargetPlatform/Private/AllDesktopTargetPlatform.cpp
@@ -81,7 +81,8 @@ FName FAllDesktopTargetPlatform::GetWaveFormat( const class USoundWave* Wave ) c
 	{
 		return NAME_ADPCM;
 	}
-	else if (Wave->IsStreaming())
+	// there is no one platform to check for Streaming status here
+	else if (Wave->IsStreaming(TEXT("Windows")))
 	{
 #if !USE_VORBIS_FOR_STREAMING
 		return NAME_OPUS;
diff --git a/Engine/Source/Developer/AllDesktopTargetPlatform/Private/AllDesktopTargetPlatform.h b/Engine/Source/Developer/AllDesktopTargetPlatform/Private/AllDesktopTargetPlatform.h
--- a/Engine/Source/Developer/AllDesktopTargetPlatform/Private/AllDesktopTargetPlatform.h
+++ b/Engine/Source/Developer/AllDesktopTargetPlatform/Private/AllDesktopTargetPlatform.h
@@ -151,10 +151,6 @@ public:
 	virtual FName GetWaveFormat( const class USoundWave* Wave ) const override;
 	virtual void GetAllWaveFormats(TArray<FName>& OutFormats) const override;
 
-	virtual FPlatformAudioCookOverrides* GetAudioCompressionSettings() const override
-	{
-		return nullptr;
-	}
 #endif // WITH_ENGINE
 
 
diff --git a/Engine/Source/Developer/Android/AndroidTargetPlatform/Private/AndroidTargetPlatform.cpp b/Engine/Source/Developer/Android/AndroidTargetPlatform/Private/AndroidTargetPlatform.cpp
--- a/Engine/Source/Developer/Android/AndroidTargetPlatform/Private/AndroidTargetPlatform.cpp
+++ b/Engine/Source/Developer/Android/AndroidTargetPlatform/Private/AndroidTargetPlatform.cpp
@@ -665,141 +665,6 @@ void FAndroidTargetPlatform::GetAllWaveFormats(TArray<FName>& OutFormats) const
 	OutFormats.Add(NAME_ADPCM);
 }
 
-namespace Android
-{
-	void CachePlatformAudioCookOverrides(FPlatformAudioCookOverrides& OutOverrides)
-	{
-		const TCHAR* CategoryName = TEXT("/Script/AndroidRuntimeSettings.AndroidRuntimeSettings");
-
-		int32 SoundCueQualityIndex = INDEX_NONE;
-		if (GConfig->GetInt(CategoryName, TEXT("SoundCueCookQualityIndex"), SoundCueQualityIndex, GEngineIni))
-		{
-			OutOverrides.SoundCueCookQualityIndex = SoundCueQualityIndex;
-		}
-
-		GConfig->GetBool(CategoryName, TEXT("bUseAudioStreamCaching"), OutOverrides.bUseStreamCaching, GEngineIni);
-
-		/** Memory Load On Demand Settings */
-		if (OutOverrides.bUseStreamCaching)
-		{
-			// Cache size:
-			int32 RetrievedCacheSize = 32 * 1024;
-			GConfig->GetInt(CategoryName, TEXT("CacheSizeKB"), RetrievedCacheSize, GEngineIni);
-			OutOverrides.StreamCachingSettings.CacheSizeKB = RetrievedCacheSize;
-		}
-
-		GConfig->GetBool(CategoryName, TEXT("bResampleForDevice"), OutOverrides.bResampleForDevice, GEngineIni);
-
-		GConfig->GetFloat(CategoryName, TEXT("CompressionQualityModifier"), OutOverrides.CompressionQualityModifier, GEngineIni);
-
-		GConfig->GetFloat(CategoryName, TEXT("AutoStreamingThreshold"), OutOverrides.AutoStreamingThreshold, GEngineIni);
-
-		//Cache sample rate map:
-		float RetrievedSampleRate = -1.0f;
-
-		GConfig->GetFloat(CategoryName, TEXT("MaxSampleRate"), RetrievedSampleRate, GEngineIni);
-		float* FoundSampleRate = OutOverrides.PlatformSampleRates.Find(ESoundwaveSampleRateSettings::Max);
-
-		if (FoundSampleRate)
-		{
-			if (!FMath::IsNearlyEqual(*FoundSampleRate, RetrievedSampleRate))
-			{
-				*FoundSampleRate = RetrievedSampleRate;
-			}
-
-		}
-		else
-		{
-			OutOverrides.PlatformSampleRates.Add(ESoundwaveSampleRateSettings::Max, RetrievedSampleRate);
-		}
-
-		RetrievedSampleRate = -1.0f;
-
-		GConfig->GetFloat(CategoryName, TEXT("HighSampleRate"), RetrievedSampleRate, GEngineIni);
-		FoundSampleRate = OutOverrides.PlatformSampleRates.Find(ESoundwaveSampleRateSettings::High);
-
-		if (FoundSampleRate)
-		{
-			if (!FMath::IsNearlyEqual(*FoundSampleRate, RetrievedSampleRate))
-			{
-				*FoundSampleRate = RetrievedSampleRate;
-			}
-
-		}
-		else
-		{
-			OutOverrides.PlatformSampleRates.Add(ESoundwaveSampleRateSettings::High, RetrievedSampleRate);
-		}
-
-
-		RetrievedSampleRate = -1.0f;
-
-		GConfig->GetFloat(CategoryName, TEXT("MedSampleRate"), RetrievedSampleRate, GEngineIni);
-		FoundSampleRate = OutOverrides.PlatformSampleRates.Find(ESoundwaveSampleRateSettings::Medium);
-
-		if (FoundSampleRate)
-		{
-			if (!FMath::IsNearlyEqual(*FoundSampleRate, RetrievedSampleRate))
-			{
-				*FoundSampleRate = RetrievedSampleRate;
-			}
-		}
-		else
-		{
-			OutOverrides.PlatformSampleRates.Add(ESoundwaveSampleRateSettings::Medium, RetrievedSampleRate);
-		}
-
-		RetrievedSampleRate = -1.0f;
-
-		GConfig->GetFloat(CategoryName, TEXT("LowSampleRate"), RetrievedSampleRate, GEngineIni);
-		FoundSampleRate = OutOverrides.PlatformSampleRates.Find(ESoundwaveSampleRateSettings::Low);
-
-		if (FoundSampleRate)
-		{
-			if (!FMath::IsNearlyEqual(*FoundSampleRate, RetrievedSampleRate))
-			{
-				*FoundSampleRate = RetrievedSampleRate;
-			}
-		}
-		else
-		{
-			OutOverrides.PlatformSampleRates.Add(ESoundwaveSampleRateSettings::Low, RetrievedSampleRate);
-		}
-
-		RetrievedSampleRate = -1.0f;
-
-		GConfig->GetFloat(CategoryName, TEXT("MinSampleRate"), RetrievedSampleRate, GEngineIni);
-		FoundSampleRate = OutOverrides.PlatformSampleRates.Find(ESoundwaveSampleRateSettings::Min);
-
-		if (FoundSampleRate)
-		{
-			if (!FMath::IsNearlyEqual(*FoundSampleRate, RetrievedSampleRate))
-			{
-				*FoundSampleRate = RetrievedSampleRate;
-			}
-		}
-		else
-		{
-			OutOverrides.PlatformSampleRates.Add(ESoundwaveSampleRateSettings::Min, RetrievedSampleRate);
-		}
-	}
-}
-
-FPlatformAudioCookOverrides* FAndroidTargetPlatform::GetAudioCompressionSettings() const
-{
-	static FPlatformAudioCookOverrides Settings;
-
-	static bool bCachedPlatformSettings = false;
-
-	if (!bCachedPlatformSettings)
-	{
-		Android::CachePlatformAudioCookOverrides(Settings);
-		bCachedPlatformSettings = true;
-	}
-
-	return &Settings;
-}
-
 #endif //WITH_ENGINE
 
 bool FAndroidTargetPlatform::SupportsVariants() const
