Integrate 11372281 from Dev-VR:
Change ARKit to use the material based method for camera image color space conversion and rendering

- The material based code path is enabled by "MATERIAL_CAMERAIMAGE_CONVERSION=1" in the Build.cs file.
  The CIContext based method can still be enabled with "MATERIAL_CAMERAIMAGE_CONVERSION=0" but it's mostly kept around for testing purpose

diff --git a/Engine/Plugins/Runtime/AR/Apple/AppleARKit/Source/AppleARKit/AppleARKit.Build.cs b/Engine/Plugins/Runtime/AR/Apple/AppleARKit/Source/AppleARKit/AppleARKit.Build.cs
--- a/Engine/Plugins/Runtime/AR/Apple/AppleARKit/Source/AppleARKit/AppleARKit.Build.cs
+++ b/Engine/Plugins/Runtime/AR/Apple/AppleARKit/Source/AppleARKit/AppleARKit.Build.cs
@@ -59,5 +59,7 @@ public class AppleARKit : ModuleRules
             string PluginPath = Utils.MakePathRelativeTo(ModuleDirectory, Target.RelativeEnginePath);
             AdditionalPropertiesForReceipt.Add("IOSPlugin", Path.Combine(PluginPath, "AppleARKit_IOS_UPL.xml"));
 		}
+		
+		PublicDefinitions.Add("MATERIAL_CAMERAIMAGE_CONVERSION=1");
 	}
 }
diff --git a/Engine/Plugins/Runtime/AR/Apple/AppleARKit/Source/AppleARKit/Private/AppleARKitFrame.cpp b/Engine/Plugins/Runtime/AR/Apple/AppleARKit/Source/AppleARKit/Private/AppleARKitFrame.cpp
--- a/Engine/Plugins/Runtime/AR/Apple/AppleARKit/Source/AppleARKit/Private/AppleARKitFrame.cpp
+++ b/Engine/Plugins/Runtime/AR/Apple/AppleARKit/Source/AppleARKit/Private/AppleARKitFrame.cpp
@@ -37,7 +37,11 @@ EARWorldMappingState ToEARWorldMappingState(ARWorldMappingStatus MapStatus)
 
 #if SUPPORTS_ARKIT_1_0
 
-FAppleARKitFrame::FAppleARKitFrame( ARFrame* InARFrame, const FVector2D MinCameraUV, const FVector2D MaxCameraUV )
+#if MATERIAL_CAMERAIMAGE_CONVERSION
+FAppleARKitFrame::FAppleARKitFrame(ARFrame* InARFrame, const FVector2D MinCameraUV, const FVector2D MaxCameraUV, CVMetalTextureCacheRef MetalTextureCache)
+#else
+FAppleARKitFrame::FAppleARKitFrame(ARFrame* InARFrame, const FVector2D MinCameraUV, const FVector2D MaxCameraUV)
+#endif
 	: Camera( InARFrame.camera )
 	, LightEstimate( InARFrame.lightEstimate )
 	, WorldMappingState(EARWorldMappingState::NotAvailable)
@@ -55,6 +59,29 @@ FAppleARKitFrame::FAppleARKitFrame( ARFrame* InARFrame, const FVector2D MinCamer
 	{
 		CameraImage = InARFrame.capturedImage;
 		CFRetain(CameraImage);
+		
+#if MATERIAL_CAMERAIMAGE_CONVERSION
+		// Update SizeX & Y
+		CapturedYImageWidth = CVPixelBufferGetWidthOfPlane( InARFrame.capturedImage, 0 );
+		CapturedYImageHeight = CVPixelBufferGetHeightOfPlane( InARFrame.capturedImage, 0 );
+		CapturedCbCrImageWidth = CVPixelBufferGetWidthOfPlane( InARFrame.capturedImage, 1 );
+		CapturedCbCrImageHeight = CVPixelBufferGetHeightOfPlane( InARFrame.capturedImage, 1 );
+		
+		// Create a metal texture from the CVPixelBufferRef. The CVMetalTextureRef will
+		// be released in the FAppleARKitFrame destructor.
+		// NOTE: On success, CapturedImage will be a new CVMetalTextureRef with a ref count of 1
+		// 		 so we don't need to CFRetain it. The corresponding CFRelease is handled in
+		//
+		CVReturn Result = CVMetalTextureCacheCreateTextureFromImage(nullptr, MetalTextureCache, InARFrame.capturedImage, nullptr, MTLPixelFormatR8Unorm, CapturedYImageWidth, CapturedYImageHeight, /*PlaneIndex*/0, &CapturedYImage);
+		check( Result == kCVReturnSuccess );
+		check( CapturedYImage );
+		check( CFGetRetainCount(CapturedYImage) == 1);
+		
+		Result = CVMetalTextureCacheCreateTextureFromImage(nullptr, MetalTextureCache, InARFrame.capturedImage, nullptr, MTLPixelFormatRG8Unorm, CapturedCbCrImageWidth, CapturedCbCrImageHeight, /*PlaneIndex*/1, &CapturedCbCrImage);
+		check( Result == kCVReturnSuccess );
+		check( CapturedCbCrImage );
+		check( CFGetRetainCount(CapturedCbCrImage) == 1);
+#endif
 	}
 
 	if (InARFrame.capturedDepthData)
@@ -110,6 +137,12 @@ FAppleARKitFrame::FAppleARKitFrame( const FAppleARKitFrame& Other )
 	: Timestamp( Other.Timestamp )
 	, CameraImage( nullptr )
 	, CameraDepth( nullptr )
+#if MATERIAL_CAMERAIMAGE_CONVERSION
+	, CapturedYImageWidth( Other.CapturedYImageWidth )
+	, CapturedYImageHeight( Other.CapturedYImageHeight )
+	, CapturedCbCrImageWidth( Other.CapturedCbCrImageWidth )
+	, CapturedCbCrImageHeight( Other.CapturedCbCrImageHeight )
+#endif
 	, Camera( Other.Camera )
 	, LightEstimate( Other.LightEstimate )
 	, WorldMappingState(Other.WorldMappingState)
@@ -124,6 +157,18 @@ FAppleARKitFrame::FAppleARKitFrame( const FAppleARKitFrame& Other )
 FAppleARKitFrame::~FAppleARKitFrame()
 {
 	// Release captured image
+#if MATERIAL_CAMERAIMAGE_CONVERSION
+	if (CapturedYImage)
+	{
+		CFRelease(CapturedYImage);
+	}
+	
+	if (CapturedCbCrImage)
+	{
+		CFRelease(CapturedCbCrImage);
+	}
+#endif
+	
 	if (CameraImage != nullptr)
 	{
 		CFRelease(CameraImage);
@@ -157,6 +202,20 @@ FAppleARKitFrame& FAppleARKitFrame::operator=( const FAppleARKitFrame& Other )
 	}
 
 	// Release outgoing image
+#if MATERIAL_CAMERAIMAGE_CONVERSION
+	if (CapturedYImage)
+	{
+		CFRelease(CapturedYImage);
+		CapturedYImage = nullptr;
+	}
+	
+	if (CapturedCbCrImage)
+	{
+		CFRelease(CapturedCbCrImage);
+		CapturedYImage = nullptr;
+	}
+#endif
+	
 	if (CameraImage != nullptr)
 	{
 		CFRelease(CameraImage);
@@ -194,6 +253,10 @@ FAppleARKitFrame& FAppleARKitFrame::operator=( const FAppleARKitFrame& Other )
 	
 	// Member-wise copy
 	Timestamp = Other.Timestamp;
+#if MATERIAL_CAMERAIMAGE_CONVERSION
+	CapturedYImageWidth = Other.CapturedYImageWidth;
+	CapturedYImageHeight = Other.CapturedYImageHeight;
+#endif
 	Camera = Other.Camera;
 	LightEstimate = Other.LightEstimate;
 	WorldMappingState = Other.WorldMappingState;
diff --git a/Engine/Plugins/Runtime/AR/Apple/AppleARKit/Source/AppleARKit/Private/AppleARKitSessionDelegate.cpp b/Engine/Plugins/Runtime/AR/Apple/AppleARKit/Source/AppleARKit/Private/AppleARKitSessionDelegate.cpp
--- a/Engine/Plugins/Runtime/AR/Apple/AppleARKit/Source/AppleARKit/Private/AppleARKitSessionDelegate.cpp
+++ b/Engine/Plugins/Runtime/AR/Apple/AppleARKit/Source/AppleARKit/Private/AppleARKitSessionDelegate.cpp
@@ -13,6 +13,10 @@
 	FAppleARKitSystem* _AppleARKitSystem;
 	TArray<FVector2D> PassthroughCameraUVs;
 	FCriticalSection CameraUVsLock;
+
+#if MATERIAL_CAMERAIMAGE_CONVERSION
+	CVMetalTextureCacheRef _metalTextureCache;
+#endif
 }
 
 
@@ -24,14 +28,41 @@
 	{
 		UE_LOG(LogAppleARKit, Log, TEXT("Delegate created with session: %p"), InAppleARKitSystem);
 		_AppleARKitSystem = InAppleARKitSystem;
+#if MATERIAL_CAMERAIMAGE_CONVERSION
+		_metalTextureCache = NULL;
+#endif
 	}
 	return self;
 }
 
+#if MATERIAL_CAMERAIMAGE_CONVERSION
+- (void)setMetalTextureCache:(CVMetalTextureCacheRef)InMetalTextureCache
+{
+	// Release current?
+	if ( _metalTextureCache != nullptr )
+	{
+		CFRelease( _metalTextureCache );
+	}
+	// Set new & retain
+	_metalTextureCache = InMetalTextureCache;
+	if ( _metalTextureCache != nullptr )
+	{
+		CFRetain( _metalTextureCache );
+	}
+}
+#endif
 #pragma mark - ARSessionDelegate Methods
 
 - (void)session:(ARSession *)session didUpdateFrame:(ARFrame *)frame 
 {
+#if MATERIAL_CAMERAIMAGE_CONVERSION
+	if (!_metalTextureCache)
+	{
+		UE_LOG(LogAppleARKit, Log, TEXT("Delegate didUpdateFrame with no valid _metalTextureCache (ignoring)"));
+		return;
+	}
+#endif
+	
 	// Update the camera UVs
 	TSharedPtr<IXRCamera, ESPMode::ThreadSafe> Camera = _AppleARKitSystem->GetXRCamera(0);
 	if (Camera)
@@ -55,7 +86,11 @@
 	}
 	
 	// Bundle results into FAppleARKitFrame
+#if MATERIAL_CAMERAIMAGE_CONVERSION
+	TSharedPtr< FAppleARKitFrame, ESPMode::ThreadSafe > AppleARKitFrame( new FAppleARKitFrame( frame, MinCameraUV, MaxCameraUV, _metalTextureCache ) );
+#else
 	TSharedPtr< FAppleARKitFrame, ESPMode::ThreadSafe > AppleARKitFrame( new FAppleARKitFrame( frame, MinCameraUV, MaxCameraUV ) );
+#endif
 	
 	// Pass result to session
 	_AppleARKitSystem->SessionDidUpdateFrame_DelegateThread( AppleARKitFrame );
diff --git a/Engine/Plugins/Runtime/AR/Apple/AppleARKit/Source/AppleARKit/Private/AppleARKitSessionDelegate.h b/Engine/Plugins/Runtime/AR/Apple/AppleARKit/Source/AppleARKit/Private/AppleARKitSessionDelegate.h
--- a/Engine/Plugins/Runtime/AR/Apple/AppleARKit/Source/AppleARKit/Private/AppleARKitSessionDelegate.h
+++ b/Engine/Plugins/Runtime/AR/Apple/AppleARKit/Source/AppleARKit/Private/AppleARKitSessionDelegate.h
@@ -16,6 +16,10 @@
 
 - (id)initWithAppleARKitSystem:(class FAppleARKitSystem*)InAppleARKitSystem;
 
+#if MATERIAL_CAMERAIMAGE_CONVERSION
+- (void)setMetalTextureCache:(CVMetalTextureCacheRef)InMetalTextureCache;
+#endif
+
 @end
 
 #endif
diff --git a/Engine/Plugins/Runtime/AR/Apple/AppleARKit/Source/AppleARKit/Private/AppleARKitSystem.cpp b/Engine/Plugins/Runtime/AR/Apple/AppleARKit/Source/AppleARKit/Private/AppleARKitSystem.cpp
--- a/Engine/Plugins/Runtime/AR/Apple/AppleARKit/Source/AppleARKit/Private/AppleARKitSystem.cpp
+++ b/Engine/Plugins/Runtime/AR/Apple/AppleARKit/Source/AppleARKit/Private/AppleARKitSystem.cpp
@@ -219,7 +219,6 @@ FAppleARKitSystem::FAppleARKitSystem()
 , DeviceOrientation(EDeviceScreenOrientation::Unknown)
 , DerivedTrackingToUnrealRotation(FRotator::ZeroRotator)
 , LightEstimate(nullptr)
-, CameraImage(nullptr)
 , CameraDepth(nullptr)
 , LastTrackedGeometry_DebugId(0)
 , FaceARSupport(nullptr)
@@ -251,7 +250,10 @@ void FAppleARKitSystem::Shutdown()
 	}
 #endif
 	CameraDepth = nullptr;
+
+#if !MATERIAL_CAMERAIMAGE_CONVERSION
 	CameraImage = nullptr;
+#endif
 	
 	PersonSegmentationImage = nullptr;
 	PersonSegmentationDepthImage = nullptr;
@@ -415,12 +417,13 @@ void FAppleARKitSystem::UpdateFrame()
 		if (GameThreadFrame.IsValid())
 		{
 #if SUPPORTS_ARKIT_1_0
+#if !MATERIAL_CAMERAIMAGE_CONVERSION
 			if (GameThreadFrame->CameraImage != nullptr)
 			{
 				// Reuse the UObjects because otherwise the time between GCs causes ARKit to be starved of resources
                 CameraImage->Init(FPlatformTime::Seconds(), GameThreadFrame->CameraImage);
 			}
-
+#endif
 			if (GameThreadFrame->CameraDepth != nullptr)
 			{
 				// Reuse the UObjects because otherwise the time between GCs causes ARKit to be starved of resources
@@ -516,7 +519,10 @@ void FAppleARKitSystem::OnBeginRendering_GameThread()
 {
 #if PLATFORM_MAC || PLATFORM_IOS
     // Queue an update on the render thread
+
+#if !MATERIAL_CAMERAIMAGE_CONVERSION
 	CameraImage->Init_RenderThread();
+#endif
 	
 	if (PersonSegmentationImage)
 	{
@@ -800,7 +806,11 @@ TArray<UARPin*> FAppleARKitSystem::OnGetAllPins() const
 
 UARTextureCameraImage* FAppleARKitSystem::OnGetCameraImage()
 {
+#if MATERIAL_CAMERAIMAGE_CONVERSION
+	return nullptr;
+#else
 	return CameraImage;
+#endif
 }
 
 UARTextureCameraDepth* FAppleARKitSystem::OnGetCameraDepth()
@@ -1222,7 +1232,9 @@ void FAppleARKitSystem::AddReferencedObjects( FReferenceCollector& Collector )
 {
 	Collector.AddReferencedObjects( TrackedGeometries );
 	Collector.AddReferencedObjects( Pins );
+#if !MATERIAL_CAMERAIMAGE_CONVERSION
 	Collector.AddReferencedObject( CameraImage );
+#endif
 	Collector.AddReferencedObject( CameraDepth );
 	Collector.AddReferencedObjects( CandidateImages );
 	Collector.AddReferencedObjects( CandidateObjects );
@@ -1273,6 +1285,8 @@ void FAppleARKitSystem::ClearTrackedGeometries()
 void FAppleARKitSystem::SetupCameraTextures()
 {
 #if SUPPORTS_ARKIT_1_0
+
+#if !MATERIAL_CAMERAIMAGE_CONVERSION
 	if (CameraImage == nullptr)
 	{
 		CameraImage = NewObject<UAppleARKitTextureCameraImage>();
@@ -1281,6 +1295,8 @@ void FAppleARKitSystem::SetupCameraTextures()
 		check(Camera);
 		Camera->SetOverlayTexture(CameraImage);
 	}
+#endif
+	
 	if (CameraDepth == nullptr)
 	{
 		CameraDepth = NewObject<UAppleARKitTextureCameraDepth>();
@@ -1342,7 +1358,23 @@ bool FAppleARKitSystem::Run(UARSessionConfig* SessionConfig)
 		{
 			Delegate = [[FAppleARKitSessionDelegate alloc] initWithAppleARKitSystem:this];
 		}
+		
+#if MATERIAL_CAMERAIMAGE_CONVERSION
+		// Create MetalTextureCache
+		if (IsMetalPlatform(GMaxRHIShaderPlatform))
+		{
+			id<MTLDevice> Device = (id<MTLDevice>)GDynamicRHI->RHIGetNativeDevice();
+			check(Device);
 
+			CVReturn Return = CVMetalTextureCacheCreate(nullptr, nullptr, Device, nullptr, &MetalTextureCache);
+			check(Return == kCVReturnSuccess);
+			check(MetalTextureCache);
+
+			// Pass to session delegate to use for Metal texture creation
+			[Delegate setMetalTextureCache : MetalTextureCache];
+		}
+#endif
+		
 		if (Session == nullptr)
 		{
 			// Start a new ARSession
@@ -1426,6 +1458,18 @@ bool FAppleARKitSystem::Pause()
 		// Suspend the session
 		[Session pause];
 		
+#if MATERIAL_CAMERAIMAGE_CONVERSION
+		// Release MetalTextureCache created in Start
+		if (MetalTextureCache)
+		{
+			// Tell delegate to release it
+			[Delegate setMetalTextureCache:nullptr];
+		
+			CFRelease(MetalTextureCache);
+			MetalTextureCache = nullptr;
+		}
+#endif
+		
 		if (CVarReleaseSessionWhenStopped.GetValueOnAnyThread())
 		{
 			[Session release];
@@ -1481,10 +1525,13 @@ void FAppleARKitSystem::SessionDidUpdateFrame_DelegateThread(TSharedPtr< FAppleA
 			WriteCameraImageToDisk(Frame->CameraImage);
 		}
 		
+#if !MATERIAL_CAMERAIMAGE_CONVERSION
 		if (CameraImage != nullptr)
 		{
 			CameraImage->EnqueueNewCameraImage(Frame->CameraImage);
 		}
+#endif
+		
 #endif
 	}
 }
diff --git a/Engine/Plugins/Runtime/AR/Apple/AppleARKit/Source/AppleARKit/Private/AppleARKitSystem.h b/Engine/Plugins/Runtime/AR/Apple/AppleARKit/Source/AppleARKit/Private/AppleARKitSystem.h
--- a/Engine/Plugins/Runtime/AR/Apple/AppleARKit/Source/AppleARKit/Private/AppleARKitSystem.h
+++ b/Engine/Plugins/Runtime/AR/Apple/AppleARKit/Source/AppleARKit/Private/AppleARKitSystem.h
@@ -176,6 +176,11 @@ private:
 	// ARKit Session Delegate
 	FAppleARKitSessionDelegate* Delegate = nullptr;
 
+#if MATERIAL_CAMERAIMAGE_CONVERSION
+	/** The Metal texture cache for unbuffered texture uploads. */
+	CVMetalTextureCacheRef MetalTextureCache = nullptr;
+#endif
+	
 	/** Cache of images that we've converted previously to prevent repeated conversion */
 	TMap< FString, CGImage* > ConvertedCandidateImages;
 
@@ -187,7 +192,9 @@ private:
 	TMap< FGuid, UARTrackedGeometry* > TrackedGeometries;
 	TArray<UARPin*> Pins;
 	UARLightEstimate* LightEstimate;
-	UAppleARKitTextureCameraImage* CameraImage;
+#if !MATERIAL_CAMERAIMAGE_CONVERSION
+	UAppleARKitTextureCameraImage* CameraImage = nullptr;
+#endif
 	UAppleARKitTextureCameraDepth* CameraDepth;
 	TMap< FString, UARCandidateImage* > CandidateImages;
 	TMap< FString, UARCandidateObject* > CandidateObjects;
diff --git a/Engine/Plugins/Runtime/AR/Apple/AppleARKit/Source/AppleARKit/Private/AppleARKitVideoOverlay.cpp b/Engine/Plugins/Runtime/AR/Apple/AppleARKit/Source/AppleARKit/Private/AppleARKitVideoOverlay.cpp
--- a/Engine/Plugins/Runtime/AR/Apple/AppleARKit/Source/AppleARKit/Private/AppleARKitVideoOverlay.cpp
+++ b/Engine/Plugins/Runtime/AR/Apple/AppleARKit/Source/AppleARKit/Private/AppleARKitVideoOverlay.cpp
@@ -35,10 +35,72 @@ DECLARE_FLOAT_COUNTER_STAT(TEXT("ARKit Frame to Render Delay (ms)"), STAT_ARKitF
 
 DECLARE_CYCLE_STAT(TEXT("Update Occlusion Textures"), STAT_UpdateOcclusionTextures, STATGROUP_ARKIT);
 
+#if MATERIAL_CAMERAIMAGE_CONVERSION
+const FString UARKitCameraOverlayMaterialLoader::OverlayMaterialPath(TEXT("/AppleARKit/ARKitCameraMaterial.ARKitCameraMaterial"));
+#else
 const FString UARKitCameraOverlayMaterialLoader::OverlayMaterialPath(TEXT("/AppleARKit/M_CameraOverlay.M_CameraOverlay"));
+#endif
+
 const FString UARKitCameraOverlayMaterialLoader::DepthOcclusionOverlayMaterialPath(TEXT("/AppleARKit/MI_DepthOcclusionOverlay.MI_DepthOcclusionOverlay"));
 const FString UARKitCameraOverlayMaterialLoader::MatteOcclusionOverlayMaterialPath(TEXT("/AppleARKit/MI_MatteOcclusionOverlay.MI_MatteOcclusionOverlay"));
 
+#if SUPPORTS_ARKIT_1_0 && MATERIAL_CAMERAIMAGE_CONVERSION
+/**
+* Passes a CVMetalTextureRef through to the RHI to wrap in an RHI texture without traversing system memory.
+* @see FAvfTexture2DResourceWrapper & FMetalSurface::FMetalSurface
+*/
+class FAppleARKitCameraTextureResourceWrapper : public FResourceBulkDataInterface
+{
+public:
+	FAppleARKitCameraTextureResourceWrapper(CFTypeRef InImageBuffer)
+		: ImageBuffer(InImageBuffer)
+	{
+		check(ImageBuffer);
+		CFRetain(ImageBuffer);
+	}
+
+	/**
+	* @return ptr to the resource memory which has been preallocated
+	*/
+	virtual const void* GetResourceBulkData() const override
+	{
+		return ImageBuffer;
+	}
+
+	/**
+	* @return size of resource memory
+	*/
+	virtual uint32 GetResourceBulkDataSize() const override
+	{
+		return 0;
+	}
+
+	/**
+	* @return the type of bulk data for special handling
+	*/
+	virtual EBulkDataType GetResourceType() const override
+	{
+		return EBulkDataType::MediaTexture;
+	}
+
+	/**
+	* Free memory after it has been used to initialize RHI resource
+	*/
+	virtual void Discard() override
+	{
+		delete this;
+	}
+
+	virtual ~FAppleARKitCameraTextureResourceWrapper()
+	{
+		CFRelease(ImageBuffer);
+		ImageBuffer = nullptr;
+	}
+
+	CFTypeRef ImageBuffer;
+};
+#endif
+
 FAppleARKitVideoOverlay::FAppleARKitVideoOverlay()
 	: MID_CameraOverlay(nullptr)
 {
@@ -165,6 +227,48 @@ using FARKitCameraOverlayMobilePS = TARKitCameraOverlayPS<true>;
 IMPLEMENT_SHADER_TYPE(template<>, FARKitCameraOverlayPS, TEXT("/Engine/Private/PostProcessMaterialShaders.usf"), TEXT("MainPS_VideoOverlay"), SF_Pixel);
 IMPLEMENT_SHADER_TYPE(template<>, FARKitCameraOverlayMobilePS, TEXT("/Engine/Private/PostProcessMaterialShaders.usf"), TEXT("MainPS"), SF_Pixel);
 
+void FAppleARKitVideoOverlay::UpdateVideoTextures(const FAppleARKitFrame& Frame)
+{
+#if SUPPORTS_ARKIT_1_0 && MATERIAL_CAMERAIMAGE_CONVERSION
+	if (FAppleARKitAvailability::SupportsARKit10())
+	{
+		check(IsMetalPlatform(GMaxRHIShaderPlatform));
+		
+		if (Frame.CapturedYImage && Frame.CapturedCbCrImage)
+		{
+			FRHIResourceCreateInfo CreateInfo;
+			const uint32 CreateFlags = TexCreate_Dynamic | TexCreate_NoTiling | TexCreate_ShaderResource;
+			CreateInfo.BulkData = new FAppleARKitCameraTextureResourceWrapper(Frame.CapturedYImage);
+			CreateInfo.ResourceArray = nullptr;
+
+			// pull the Y and CbCr textures out of the captured image planes (format is fake here, it will get the format from the FAppleARKitCameraTextureResourceWrapper)
+			VideoTextureY = RHICreateTexture2D(Frame.CapturedYImageWidth, Frame.CapturedYImageHeight, /*Format=*/PF_B8G8R8A8, /*NumMips=*/1, /*NumSamples=*/1, CreateFlags, CreateInfo);
+
+			CreateInfo.BulkData = new FAppleARKitCameraTextureResourceWrapper(Frame.CapturedCbCrImage);
+			VideoTextureCbCr = RHICreateTexture2D(Frame.CapturedCbCrImageWidth, Frame.CapturedCbCrImageHeight, /*Format=*/PF_B8G8R8A8, /*NumMips=*/1, /*NumSamples=*/1, CreateFlags, CreateInfo);
+
+			// todo: Add an update call to the registry instead of this unregister/re-register
+			FExternalTextureRegistry::Get().UnregisterExternalTexture(ARKitPassthroughCameraExternalTextureYGuid);
+			FExternalTextureRegistry::Get().UnregisterExternalTexture(ARKitPassthroughCameraExternalTextureCbCrGuid);
+
+			FSamplerStateInitializerRHI SamplerStateInitializer(SF_Bilinear, AM_Wrap, AM_Wrap, AM_Wrap);
+			FSamplerStateRHIRef SamplerStateRHI = RHICreateSamplerState(SamplerStateInitializer);
+
+			FExternalTextureRegistry::Get().RegisterExternalTexture(ARKitPassthroughCameraExternalTextureYGuid, VideoTextureY, SamplerStateRHI);
+			FExternalTextureRegistry::Get().RegisterExternalTexture(ARKitPassthroughCameraExternalTextureCbCrGuid, VideoTextureCbCr, SamplerStateRHI);
+			
+			//Make sure AR camera pass through materials are updated properly
+			FMaterialRenderProxy::UpdateDeferredCachedUniformExpressions();
+
+			//CFRelease(Frame.CapturedYImage);
+			//CFRelease(Frame.CapturedCbCrImage);
+			//Frame.CapturedYImage = nullptr;
+			//Frame.CapturedCbCrImage = nullptr;
+		}
+	}
+#endif
+}
+
 void FAppleARKitVideoOverlay::UpdateOcclusionTextures(const FAppleARKitFrame& Frame)
 {
 #if SUPPORTS_ARKIT_3_0
@@ -215,7 +319,7 @@ void FAppleARKitVideoOverlay::UpdateOcclusionTextures(const FAppleARKitFrame& Fr
 #endif
 }
 
-void FAppleARKitVideoOverlay::RenderVideoOverlayWithMaterial(FRHICommandListImmediate& RHICmdList, const FSceneView& InView, struct FAppleARKitFrame& Frame, const EDeviceScreenOrientation DeviceOrientation, UMaterialInstanceDynamic* RenderingOverlayMaterial, const bool bRenderingOcclusion)
+void FAppleARKitVideoOverlay::RenderVideoOverlayWithMaterial(FRHICommandListImmediate& RHICmdList, const FSceneView& InView, struct FAppleARKitFrame& Frame, const EDeviceScreenOrientation DeviceOrientation, UMaterialInstanceDynamic* RenderingOverlayMaterial, const bool bRenderingOcclusion, const bool bRespectAllOrientations)
 {
 	if (RenderingOverlayMaterial == nullptr || !RenderingOverlayMaterial->IsValidLowLevel())
 	{
@@ -230,7 +334,7 @@ void FAppleARKitVideoOverlay::RenderVideoOverlayWithMaterial(FRHICommandListImme
 	
 	SCOPED_DRAW_EVENTF(RHICmdList, RenderVideoOverlay, bRenderingOcclusion ? TEXT("VideoOverlay (Occlusion)") : TEXT("VideoOverlay (Background)"));
 	
-	if (OverlayVertexBufferRHI[0] == nullptr || !OverlayVertexBufferRHI[0].IsValid())
+	if (OverlayVertexBufferRHITwoOrientations[0] == nullptr || !OverlayVertexBufferRHITwoOrientations[0].IsValid())
 	{
 		const FVector2D ViewSize(InView.UnconstrainedViewRect.Max.X, InView.UnconstrainedViewRect.Max.Y);
 
@@ -268,43 +372,100 @@ void FAppleARKitVideoOverlay::RenderVideoOverlayWithMaterial(FRHICommandListImme
 			FVector4(1.0f, 0.0f, 0.0f, 1.0f)
 		};
 
-		const FVector2D UVs[] =
+		// 2 orientations
 		{
-			// Landscape
-			FVector2D(UVOffset.X, 1.0f - UVOffset.Y),
-			FVector2D(UVOffset.X, UVOffset.Y),
-			FVector2D(1.0f - UVOffset.X, 1.0f - UVOffset.Y),
-			FVector2D(1.0f - UVOffset.X, UVOffset.Y),
-            
-			// Portrait
-			FVector2D(UVOffset.Y, 1.0f - UVOffset.X),
-            FVector2D(UVOffset.Y, UVOffset.X),
-			FVector2D(1.0f - UVOffset.Y, 1.0f - UVOffset.X),
-			FVector2D(1.0f - UVOffset.Y, UVOffset.X),
-		};
+			const FVector2D UVs[] =
+			{
+				// Landscape
+				FVector2D(UVOffset.X, 1.0f - UVOffset.Y),
+				FVector2D(UVOffset.X, UVOffset.Y),
+				FVector2D(1.0f - UVOffset.X, 1.0f - UVOffset.Y),
+				FVector2D(1.0f - UVOffset.X, UVOffset.Y),
+				
+				// Portrait
+				FVector2D(UVOffset.Y, 1.0f - UVOffset.X),
+				FVector2D(UVOffset.Y, UVOffset.X),
+				FVector2D(1.0f - UVOffset.Y, 1.0f - UVOffset.X),
+				FVector2D(1.0f - UVOffset.Y, UVOffset.X),
+			};
+
+			uint32 UVIndex = 0;
+			for (uint32 OrientationIter = 0; OrientationIter < 2; ++OrientationIter)
+			{
+				TResourceArray<FFilterVertex, VERTEXBUFFER_ALIGNMENT> Vertices;
+				Vertices.SetNumUninitialized(4);
+
+				Vertices[0].Position = Positions[0];
+				Vertices[0].UV = UVs[UVIndex];
 
-		uint32 UVIndex = 0;
-		for (uint32 OrientationIter = 0; OrientationIter < 2; ++OrientationIter)
+				Vertices[1].Position = Positions[1];
+				Vertices[1].UV = UVs[UVIndex + 1];
+
+				Vertices[2].Position = Positions[2];
+				Vertices[2].UV = UVs[UVIndex + 2];
+
+				Vertices[3].Position = Positions[3];
+				Vertices[3].UV = UVs[UVIndex + 3];
+
+				UVIndex += 4;
+
+				FRHIResourceCreateInfo CreateInfoVB(&Vertices);
+				OverlayVertexBufferRHITwoOrientations[OrientationIter] = RHICreateVertexBuffer(Vertices.GetResourceDataSize(), BUF_Static, CreateInfoVB);
+			}
+		}
+		
+		// 4 orientations
 		{
-			TResourceArray<FFilterVertex, VERTEXBUFFER_ALIGNMENT> Vertices;
-			Vertices.SetNumUninitialized(4);
+			const FVector2D UVs[] =
+			{
+				// Landscape left
+				FVector2D(UVOffset.X, 1.0f - UVOffset.Y),
+				FVector2D(UVOffset.X, UVOffset.Y),
+				FVector2D(1.0f - UVOffset.X, 1.0f - UVOffset.Y),
+				FVector2D(1.0f - UVOffset.X, UVOffset.Y),
+
+				// Landscape right
+				FVector2D(1.0f - UVOffset.X, UVOffset.Y),
+				FVector2D(1.0f - UVOffset.X, 1.0f - UVOffset.Y),
+				FVector2D(UVOffset.X, UVOffset.Y),
+				FVector2D(UVOffset.X, 1.0f - UVOffset.Y),
+
+				// Portrait
+				FVector2D(1.0f - UVOffset.X, 1.0f - UVOffset.Y),
+				FVector2D(UVOffset.X, 1.0f - UVOffset.Y),
+				FVector2D(1.0f - UVOffset.X, UVOffset.Y),
+				FVector2D(UVOffset.X, UVOffset.Y),
+
+				// Portrait Upside Down
+				FVector2D(UVOffset.X, UVOffset.Y),
+				FVector2D(1.0f - UVOffset.X, UVOffset.Y),
+				FVector2D(UVOffset.X, 1.0f - UVOffset.Y),
+				FVector2D(1.0f - UVOffset.X, 1.0f - UVOffset.Y)
+			};
+
+			uint32 UVIndex = 0;
+			for (uint32 OrientationIter = 0; OrientationIter < 4; ++OrientationIter)
+			{
+				TResourceArray<FFilterVertex, VERTEXBUFFER_ALIGNMENT> Vertices;
+				Vertices.SetNumUninitialized(4);
 
-			Vertices[0].Position = Positions[0];
-			Vertices[0].UV = UVs[UVIndex];
+				Vertices[0].Position = Positions[0];
+				Vertices[0].UV = UVs[UVIndex];
 
-			Vertices[1].Position = Positions[1];
-			Vertices[1].UV = UVs[UVIndex + 1];
+				Vertices[1].Position = Positions[1];
+				Vertices[1].UV = UVs[UVIndex + 1];
 
-			Vertices[2].Position = Positions[2];
-			Vertices[2].UV = UVs[UVIndex + 2];
+				Vertices[2].Position = Positions[2];
+				Vertices[2].UV = UVs[UVIndex + 2];
 
-			Vertices[3].Position = Positions[3];
-			Vertices[3].UV = UVs[UVIndex + 3];
+				Vertices[3].Position = Positions[3];
+				Vertices[3].UV = UVs[UVIndex + 3];
 
-			UVIndex += 4;
+				UVIndex += 4;
 
-			FRHIResourceCreateInfo CreateInfoVB(&Vertices);
-			OverlayVertexBufferRHI[OrientationIter] = RHICreateVertexBuffer(Vertices.GetResourceDataSize(), BUF_Static, CreateInfoVB);
+				FRHIResourceCreateInfo CreateInfoVB(&Vertices);
+				OverlayVertexBufferRHIFourOrientations[OrientationIter] = RHICreateVertexBuffer(Vertices.GetResourceDataSize(), BUF_Static, CreateInfoVB);
+			}
 		}
 	}
 
@@ -401,21 +562,49 @@ void FAppleARKitVideoOverlay::RenderVideoOverlayWithMaterial(FRHICommandListImme
 	}
 
 	FRHIVertexBuffer* VertexBufferRHI = nullptr;
-	switch (DeviceOrientation)
+	if (bRespectAllOrientations)
 	{
-		case EDeviceScreenOrientation::LandscapeRight:
-		case EDeviceScreenOrientation::LandscapeLeft:
-			VertexBufferRHI = OverlayVertexBufferRHI[0];
-			break;
-
-		case EDeviceScreenOrientation::Portrait:
-		case EDeviceScreenOrientation::PortraitUpsideDown:
-			VertexBufferRHI = OverlayVertexBufferRHI[1];
-			break;
-
-		default:
-			VertexBufferRHI = OverlayVertexBufferRHI[0];
-			break;
+		switch (DeviceOrientation)
+		{
+			case EDeviceScreenOrientation::LandscapeRight:
+				VertexBufferRHI = OverlayVertexBufferRHIFourOrientations[0];
+				break;
+
+			case EDeviceScreenOrientation::LandscapeLeft:
+				VertexBufferRHI = OverlayVertexBufferRHIFourOrientations[1];
+				break;
+
+			case EDeviceScreenOrientation::Portrait:
+				VertexBufferRHI = OverlayVertexBufferRHIFourOrientations[2];
+				break;
+
+			case EDeviceScreenOrientation::PortraitUpsideDown:
+				VertexBufferRHI = OverlayVertexBufferRHIFourOrientations[3];
+				break;
+				
+			default:
+				VertexBufferRHI = OverlayVertexBufferRHIFourOrientations[0];
+				break;
+		}
+	}
+	else
+	{
+		switch (DeviceOrientation)
+		{
+			case EDeviceScreenOrientation::LandscapeRight:
+			case EDeviceScreenOrientation::LandscapeLeft:
+				VertexBufferRHI = OverlayVertexBufferRHITwoOrientations[0];
+				break;
+
+			case EDeviceScreenOrientation::Portrait:
+			case EDeviceScreenOrientation::PortraitUpsideDown:
+				VertexBufferRHI = OverlayVertexBufferRHITwoOrientations[1];
+				break;
+
+			default:
+				VertexBufferRHI = OverlayVertexBufferRHITwoOrientations[0];
+				break;
+		}
 	}
 
 
@@ -436,9 +625,14 @@ void FAppleARKitVideoOverlay::RenderVideoOverlayWithMaterial(FRHICommandListImme
 
 void FAppleARKitVideoOverlay::RenderVideoOverlay_RenderThread(FRHICommandListImmediate& RHICmdList, const FSceneView& InView, FAppleARKitFrame& Frame, const EDeviceScreenOrientation DeviceOrientation, const float WorldToMeterScale)
 {
+	UpdateVideoTextures(Frame);
 	UpdateOcclusionTextures(Frame);
 	
-	RenderVideoOverlayWithMaterial(RHICmdList, InView, Frame, DeviceOrientation, MID_CameraOverlay, false);
+#if MATERIAL_CAMERAIMAGE_CONVERSION
+	RenderVideoOverlayWithMaterial(RHICmdList, InView, Frame, DeviceOrientation, MID_CameraOverlay, false, true);
+#else
+	RenderVideoOverlayWithMaterial(RHICmdList, InView, Frame, DeviceOrientation, MID_CameraOverlay, false, false);
+#endif
 	
 	if (bEnablePersonOcclusion)
 	{
@@ -458,7 +652,7 @@ void FAppleARKitVideoOverlay::RenderVideoOverlay_RenderThread(FRHICommandListImm
 		static const FName DepthTextureUVParamName(TEXT("DepthTextureUVParam"));
 		OcclusionMaterial->SetVectorParameterValue(DepthTextureUVParamName, DepthTextureUVParamValue);
 		
-		RenderVideoOverlayWithMaterial(RHICmdList, InView, Frame, DeviceOrientation, OcclusionMaterial, true);
+		RenderVideoOverlayWithMaterial(RHICmdList, InView, Frame, DeviceOrientation, OcclusionMaterial, true, false);
 		bOcclusionDepthTextureRecentlyUpdated = false;
 	}
 }
@@ -466,7 +660,8 @@ void FAppleARKitVideoOverlay::RenderVideoOverlay_RenderThread(FRHICommandListImm
 bool FAppleARKitVideoOverlay::GetPassthroughCameraUVs_RenderThread(TArray<FVector2D>& OutUVs, const EDeviceScreenOrientation DeviceOrientation)
 {
 #if SUPPORTS_ARKIT_1_0
-	if (OverlayVertexBufferRHI[0] != nullptr && OverlayVertexBufferRHI[0].IsValid())
+	// UVOffset is initialized when OverlayVertexBufferRHITwoOrientations is created
+	if (OverlayVertexBufferRHITwoOrientations[0] != nullptr && OverlayVertexBufferRHITwoOrientations[0].IsValid())
 	{
 		OutUVs.SetNumUninitialized(4);
 
diff --git a/Engine/Plugins/Runtime/AR/Apple/AppleARKit/Source/AppleARKit/Private/AppleARKitVideoOverlay.h b/Engine/Plugins/Runtime/AR/Apple/AppleARKit/Source/AppleARKit/Private/AppleARKitVideoOverlay.h
--- a/Engine/Plugins/Runtime/AR/Apple/AppleARKit/Source/AppleARKit/Private/AppleARKitVideoOverlay.h
+++ b/Engine/Plugins/Runtime/AR/Apple/AppleARKit/Source/AppleARKit/Private/AppleARKitVideoOverlay.h
@@ -70,15 +70,17 @@ private:
 	virtual void AddReferencedObjects(FReferenceCollector& Collector) override;
 	//~ FGCObject
 	
-	void RenderVideoOverlayWithMaterial(FRHICommandListImmediate& RHICmdList, const FSceneView& InView, struct FAppleARKitFrame& Frame, const EDeviceScreenOrientation DeviceOrientation, UMaterialInstanceDynamic* RenderingOverlayMaterial, const bool bRenderingOcclusion);
+	void RenderVideoOverlayWithMaterial(FRHICommandListImmediate& RHICmdList, const FSceneView& InView, struct FAppleARKitFrame& Frame, const EDeviceScreenOrientation DeviceOrientation, UMaterialInstanceDynamic* RenderingOverlayMaterial, const bool bRenderingOcclusion, const bool bRespectAllOrientations);
 	void UpdateOcclusionTextures(const FAppleARKitFrame& Frame);
+	void UpdateVideoTextures(const FAppleARKitFrame& Frame);
 
 	UMaterialInstanceDynamic* MID_CameraOverlay;
 
 	// Cache UV offset to be used by GetPassthroughCameraUVs_RenderThread
 	FVector2D UVOffset;
-
-	FVertexBufferRHIRef OverlayVertexBufferRHI[2];
+	
+	FVertexBufferRHIRef OverlayVertexBufferRHITwoOrientations[2];
+	FVertexBufferRHIRef OverlayVertexBufferRHIFourOrientations[4];
 	FIndexBufferRHIRef IndexBufferRHI;
 	
 	bool bEnablePersonOcclusion = false;
@@ -93,4 +95,9 @@ private:
 	UAppleARKitOcclusionTexture* OcclusionDepthTexture = nullptr;
 	UMaterialInstanceDynamic* MID_DepthOcclusionOverlay = nullptr;
 	UMaterialInstanceDynamic* MID_MatteOcclusionOverlay = nullptr;
+	
+#if MATERIAL_CAMERAIMAGE_CONVERSION
+	FTextureRHIRef VideoTextureY;
+	FTextureRHIRef VideoTextureCbCr;
+#endif
 };
diff --git a/Engine/Plugins/Runtime/AR/Apple/AppleARKit/Source/AppleARKit/Private/ExternalTextureMaterialExpression.cpp b/Engine/Plugins/Runtime/AR/Apple/AppleARKit/Source/AppleARKit/Private/ExternalTextureMaterialExpression.cpp
--- a/Engine/Plugins/Runtime/AR/Apple/AppleARKit/Source/AppleARKit/Private/ExternalTextureMaterialExpression.cpp
+++ b/Engine/Plugins/Runtime/AR/Apple/AppleARKit/Source/AppleARKit/Private/ExternalTextureMaterialExpression.cpp
@@ -4,12 +4,12 @@
 #include "MaterialCompiler.h"
 #include "ExternalTextureGuid.h"
 
-UDEPRECATED_MaterialExpressionARKitPassthroughCamera::UDEPRECATED_MaterialExpressionARKitPassthroughCamera(const FObjectInitializer& ObjectInitializer)
+UMaterialExpressionARKitPassthroughCamera::UMaterialExpressionARKitPassthroughCamera(const FObjectInitializer& ObjectInitializer)
 	: Super(ObjectInitializer)
 {}
 
 #if WITH_EDITOR
-int32 UDEPRECATED_MaterialExpressionARKitPassthroughCamera::Compile(class FMaterialCompiler* Compiler, int32 OutputIndex)
+int32 UMaterialExpressionARKitPassthroughCamera::Compile(class FMaterialCompiler* Compiler, int32 OutputIndex)
 {
 	return Compiler->TextureSample(
 		Compiler->ExternalTexture((TextureType == TextureY) ? ARKitPassthroughCameraExternalTextureYGuid : ARKitPassthroughCameraExternalTextureCbCrGuid),
@@ -17,13 +17,13 @@ int32 UDEPRECATED_MaterialExpressionARKitPassthroughCamera::Compile(class FMater
 		EMaterialSamplerType::SAMPLERTYPE_Color);
 }
 
-int32 UDEPRECATED_MaterialExpressionARKitPassthroughCamera::CompilePreview(class FMaterialCompiler* Compiler, int32 OutputIndex)
+int32 UMaterialExpressionARKitPassthroughCamera::CompilePreview(class FMaterialCompiler* Compiler, int32 OutputIndex)
 {
 	return INDEX_NONE;
 }
 
-void UDEPRECATED_MaterialExpressionARKitPassthroughCamera::GetCaption(TArray<FString>& OutCaptions) const
+void UMaterialExpressionARKitPassthroughCamera::GetCaption(TArray<FString>& OutCaptions) const
 {
-	OutCaptions.Add(TEXT("DEPRECATED USE ARSystem->GetCameraImage() - ARKit Passthrough Camera"));
+	OutCaptions.Add(TEXT("ARKit Passthrough Camera"));
 }
-#endif
\ No newline at end of file
+#endif
diff --git a/Engine/Plugins/Runtime/AR/Apple/AppleARKit/Source/AppleARKit/Private/ExternalTextureMaterialExpression.h b/Engine/Plugins/Runtime/AR/Apple/AppleARKit/Source/AppleARKit/Private/ExternalTextureMaterialExpression.h
--- a/Engine/Plugins/Runtime/AR/Apple/AppleARKit/Source/AppleARKit/Private/ExternalTextureMaterialExpression.h
+++ b/Engine/Plugins/Runtime/AR/Apple/AppleARKit/Source/AppleARKit/Private/ExternalTextureMaterialExpression.h
@@ -17,8 +17,8 @@ enum EARKitTextureType
 /**
 * Implements a node sampling from the ARKit Passthrough external textures.
 */
-UCLASS(collapsecategories, hidecategories = Object, Deprecated)
-class UDEPRECATED_MaterialExpressionARKitPassthroughCamera : public UMaterialExpression
+UCLASS(collapsecategories, hidecategories = Object)
+class UMaterialExpressionARKitPassthroughCamera : public UMaterialExpression
 {
 	GENERATED_UCLASS_BODY()
 
diff --git a/Engine/Plugins/Runtime/AR/Apple/AppleARKit/Source/AppleARKit/Public/AppleARKitFrame.h b/Engine/Plugins/Runtime/AR/Apple/AppleARKit/Source/AppleARKit/Public/AppleARKitFrame.h
--- a/Engine/Plugins/Runtime/AR/Apple/AppleARKit/Source/AppleARKit/Public/AppleARKitFrame.h
+++ b/Engine/Plugins/Runtime/AR/Apple/AppleARKit/Source/AppleARKit/Public/AppleARKitFrame.h
@@ -31,7 +31,11 @@ struct APPLEARKIT_API FAppleARKitFrame
 	 * @param MinCameraUV 		- The minimum (top left) UV used to render the passthrough camera
 	 * @param MaxCameraUV 		- The maximum (bottom right) UV used to render the passthrough camera
 	 */
-	FAppleARKitFrame( ARFrame* InARFrame, const FVector2D MinCameraUV, const FVector2D MaxCameraUV );
+#if MATERIAL_CAMERAIMAGE_CONVERSION
+	FAppleARKitFrame(ARFrame* InARFrame, const FVector2D MinCameraUV, const FVector2D MaxCameraUV, CVMetalTextureCacheRef MetalTextureCache);
+#else
+	FAppleARKitFrame(ARFrame* InARFrame, const FVector2D MinCameraUV, const FVector2D MaxCameraUV);
+#endif
 
 	/** 
 	 * Copy contructor. CapturedImage is skipped as we don't need / want to retain access to the 
@@ -55,12 +59,28 @@ struct APPLEARKIT_API FAppleARKitFrame
 #if SUPPORTS_ARKIT_1_0
 
 	/** The raw camera buffer from ARKit */
-	CVPixelBufferRef CameraImage;
+	CVPixelBufferRef CameraImage = nullptr;
 	/** The raw camera depth info from ARKit (needs iPhone X) */
-	AVDepthData* CameraDepth;
-	void* NativeFrame;
+	AVDepthData* CameraDepth = nullptr;
+	void* NativeFrame = nullptr;
+	
+#if MATERIAL_CAMERAIMAGE_CONVERSION
+	/** The frame's captured images */
+    CVMetalTextureRef CapturedYImage = nullptr;
+    CVMetalTextureRef CapturedCbCrImage = nullptr;
 #endif
-
+	
+#endif
+	
+#if MATERIAL_CAMERAIMAGE_CONVERSION
+	/** The width and height in pixels of the frame's captured images. */
+	uint32 CapturedYImageWidth = 0;
+	uint32 CapturedYImageHeight = 0;
+    
+    uint32 CapturedCbCrImageWidth = 0;
+    uint32 CapturedCbCrImageHeight = 0;
+#endif
+	
 	/** The camera used to capture the frame's image. */
 	FAppleARKitCamera Camera;
 
