Fixing PVS 7.7 Issues:

warning V1043: A global object variable 'XYZ' is declared in the header. Multiple copies of it will be created in all translation units that include this header

When we are 100% C++17, we could use the inline keyword instead of making things extern.

diff --git a/Engine/Plugins/Experimental/PanoramicCapture/Source/PanoramicCapture/Private/SceneCapturer.cpp b/Engine/Plugins/Experimental/PanoramicCapture/Source/PanoramicCapture/Private/SceneCapturer.cpp
--- a/Engine/Plugins/Experimental/PanoramicCapture/Source/PanoramicCapture/Private/SceneCapturer.cpp
+++ b/Engine/Plugins/Experimental/PanoramicCapture/Source/PanoramicCapture/Private/SceneCapturer.cpp
@@ -33,6 +33,8 @@
 
 #define LOCTEXT_NAMESPACE "LogStereoPanorama"
 
+const FName StereoPanoramaLogName("LogStereoPanorama");
+
 // always combine both eyes
 const bool CombineAtlasesOnOutput = true;
 
diff --git a/Engine/Plugins/Experimental/PanoramicCapture/Source/PanoramicCapture/Private/SceneCapturer.h b/Engine/Plugins/Experimental/PanoramicCapture/Source/PanoramicCapture/Private/SceneCapturer.h
--- a/Engine/Plugins/Experimental/PanoramicCapture/Source/PanoramicCapture/Private/SceneCapturer.h
+++ b/Engine/Plugins/Experimental/PanoramicCapture/Source/PanoramicCapture/Private/SceneCapturer.h
@@ -21,7 +21,7 @@ DECLARE_CYCLE_STAT( TEXT( "SampleSpherical" ), STAT_SPSampleSpherical, STATGROUP
 DECLARE_CYCLE_STAT( TEXT( "ReadStrip" ),       STAT_SPReadStrip,       STATGROUP_SP );
 DECLARE_CYCLE_STAT( TEXT( "FillAlpha" ),       STAT_SPFillAlpha,       STATGROUP_SP );
 
-const static FName StereoPanoramaLogName("LogStereoPanorama");
+extern const FName StereoPanoramaLogName;
 
 USTRUCT()
 struct FPostProcessVolumeData
diff --git a/Engine/Plugins/Online/OnlineSubsystem/Source/Private/OnlineSubsystem.cpp b/Engine/Plugins/Online/OnlineSubsystem/Source/Private/OnlineSubsystem.cpp
--- a/Engine/Plugins/Online/OnlineSubsystem/Source/Private/OnlineSubsystem.cpp
+++ b/Engine/Plugins/Online/OnlineSubsystem/Source/Private/OnlineSubsystem.cpp
@@ -70,6 +70,9 @@ const FString CustomPresenceDataKey = TEXT("CustomData");
 /** Name of the client that sent the presence update */
 const FString DefaultAppIdKey = TEXT("AppId");
 
+/** User friendly name of the client that sent the presence update */
+const FString DefaultProductNameKey = TEXT("ProductName");
+
 /** Platform of the client that sent the presence update */
 const FString DefaultPlatformKey = TEXT("Platform");
 
diff --git a/Engine/Plugins/Online/OnlineSubsystem/Source/Public/Interfaces/OnlinePresenceInterface.h b/Engine/Plugins/Online/OnlineSubsystem/Source/Public/Interfaces/OnlinePresenceInterface.h
--- a/Engine/Plugins/Online/OnlineSubsystem/Source/Public/Interfaces/OnlinePresenceInterface.h
+++ b/Engine/Plugins/Online/OnlineSubsystem/Source/Public/Interfaces/OnlinePresenceInterface.h
@@ -36,7 +36,7 @@ extern ONLINESUBSYSTEM_API const FString CustomPresenceDataKey;
 extern ONLINESUBSYSTEM_API const FString DefaultAppIdKey;
 
 /** User friendly name of the client that sent the presence update */
-const FString DefaultProductNameKey = TEXT("ProductName");
+extern ONLINESUBSYSTEM_API const FString DefaultProductNameKey;
 
 /** Platform of the client that sent the presence update */
 extern ONLINESUBSYSTEM_API const FString DefaultPlatformKey;
diff --git a/Engine/Plugins/Runtime/AR/AppleAR/AppleARKit/Source/AppleARKit/Private/ExternalTextureGuid.cpp b/Engine/Plugins/Runtime/AR/AppleAR/AppleARKit/Source/AppleARKit/Private/ExternalTextureGuid.cpp
new file mode 100644
--- /dev/null
+++ b/Engine/Plugins/Runtime/AR/AppleAR/AppleARKit/Source/AppleARKit/Private/ExternalTextureGuid.cpp
@@ -0,0 +1,12 @@
+// Copyright Epic Games, Inc. All Rights Reserved.
+
+#include "ExternalTextureGuid.h"
+
+/* GUID used to indicate that this external texture is the ARKit camera Y component. 0a417e66-bcda-4610-8594-92248c914a20 */
+const FGuid ARKitPassthroughCameraExternalTextureYGuid(0x0A417E66, 0xBCDA4610, 0x85949224, 0x8C914A20);
+
+/* GUID used to indicate that this external texture is the ARKit camera CbCr component. d61a2708-55fe-4cb1-9cde-df2a9bb2a2d2 */
+const FGuid ARKitPassthroughCameraExternalTextureCbCrGuid(0xD61A2708, 0x55FE4CB1, 0x9CDEDF2A, 0x9BB2A2D2);
+
+/* GUID used to indicate that this external texture is the ARKit camera depth texture. cba06c09-03b2-4818-8913-6424d99acf28 */
+const FGuid ARKitPassthroughCameraExternalTextureDepthGuid(0xCBA06C09, 0x03B24818, 0x89136424, 0xD99ACF28);
diff --git a/Engine/Plugins/Runtime/AR/AppleAR/AppleARKit/Source/AppleARKit/Private/ExternalTextureGuid.h b/Engine/Plugins/Runtime/AR/AppleAR/AppleARKit/Source/AppleARKit/Private/ExternalTextureGuid.h
--- a/Engine/Plugins/Runtime/AR/AppleAR/AppleARKit/Source/AppleARKit/Private/ExternalTextureGuid.h
+++ b/Engine/Plugins/Runtime/AR/AppleAR/AppleARKit/Source/AppleARKit/Private/ExternalTextureGuid.h
@@ -6,10 +6,10 @@
 #include "Misc/Guid.h"
 
 /* GUID used to indicate that this external texture is the ARKit camera Y component. 0a417e66-bcda-4610-8594-92248c914a20 */
-static const FGuid ARKitPassthroughCameraExternalTextureYGuid(0x0A417E66, 0xBCDA4610, 0x85949224, 0x8C914A20);
+extern const FGuid ARKitPassthroughCameraExternalTextureYGuid;
 
 /* GUID used to indicate that this external texture is the ARKit camera CbCr component. d61a2708-55fe-4cb1-9cde-df2a9bb2a2d2 */
-static const FGuid ARKitPassthroughCameraExternalTextureCbCrGuid(0xD61A2708, 0x55FE4CB1, 0x9CDEDF2A, 0x9BB2A2D2);
+extern const FGuid ARKitPassthroughCameraExternalTextureCbCrGuid;
 
 /* GUID used to indicate that this external texture is the ARKit camera depth texture. cba06c09-03b2-4818-8913-6424d99acf28 */
-static const FGuid ARKitPassthroughCameraExternalTextureDepthGuid(0xCBA06C09, 0x03B24818, 0x89136424, 0xD99ACF28);
+extern const FGuid ARKitPassthroughCameraExternalTextureDepthGuid;
diff --git a/Engine/Plugins/Runtime/Oculus/OculusVR/Source/OculusInput/Private/OculusHandComponent.cpp b/Engine/Plugins/Runtime/Oculus/OculusVR/Source/OculusInput/Private/OculusHandComponent.cpp
--- a/Engine/Plugins/Runtime/Oculus/OculusVR/Source/OculusInput/Private/OculusHandComponent.cpp
+++ b/Engine/Plugins/Runtime/Oculus/OculusVR/Source/OculusInput/Private/OculusHandComponent.cpp
@@ -8,6 +8,8 @@
 
 #include "GameFramework/PlayerController.h"
 
+const FQuat HandRootFixupRotation = FQuat(-0.5f, -0.5f, 0.5f, 0.5f);
+
 UOculusHandComponent::UOculusHandComponent(const FObjectInitializer& ObjectInitializer)
 	: Super(ObjectInitializer)
 {
diff --git a/Engine/Plugins/Runtime/Oculus/OculusVR/Source/OculusInput/Public/OculusHandComponent.h b/Engine/Plugins/Runtime/Oculus/OculusVR/Source/OculusInput/Public/OculusHandComponent.h
--- a/Engine/Plugins/Runtime/Oculus/OculusVR/Source/OculusInput/Public/OculusHandComponent.h
+++ b/Engine/Plugins/Runtime/Oculus/OculusVR/Source/OculusInput/Public/OculusHandComponent.h
@@ -19,7 +19,7 @@ enum class ESystemGestureBehavior : uint8
 	SwapMaterial
 };
 
-static const FQuat HandRootFixupRotation = FQuat(-0.5f, -0.5f, 0.5f, 0.5f);
+OCULUSINPUT_API extern const FQuat HandRootFixupRotation;
 
 UCLASS(Blueprintable, meta = (BlueprintSpawnableComponent), ClassGroup = OculusHand)
 class OCULUSINPUT_API UOculusHandComponent : public UPoseableMeshComponent
diff --git a/Engine/Source/Editor/ClothPainter/Private/ClothPaintingModule.cpp b/Engine/Source/Editor/ClothPainter/Private/ClothPaintingModule.cpp
--- a/Engine/Source/Editor/ClothPainter/Private/ClothPaintingModule.cpp
+++ b/Engine/Source/Editor/ClothPainter/Private/ClothPaintingModule.cpp
@@ -25,6 +25,8 @@
 
 #define LOCTEXT_NAMESPACE "ClothPaintingModule"
 
+const FName PaintModeID = "ClothPaintMode";
+
 IMPLEMENT_MODULE(FClothPaintingModule, ClothPainter);
 
 DECLARE_DELEGATE_OneParam(FOnToggleClothPaintMode, bool);
diff --git a/Engine/Source/Editor/ClothPainter/Public/ClothPaintingModule.h b/Engine/Source/Editor/ClothPainter/Public/ClothPaintingModule.h
--- a/Engine/Source/Editor/ClothPainter/Public/ClothPaintingModule.h
+++ b/Engine/Source/Editor/ClothPainter/Public/ClothPaintingModule.h
@@ -11,7 +11,7 @@ class UEditorExperimentalSettings;
 class ISkeletalMeshEditor;
 class SClothPaintTab;
 
-const static FName PaintModeID = "ClothPaintMode";
+CLOTHPAINTER_API extern const FName PaintModeID;
 
 class FClothPaintingModule : public IModuleInterface
 {
diff --git a/Engine/Source/Editor/UndoHistory/Private/UndoHistoryModule.cpp b/Engine/Source/Editor/UndoHistory/Private/UndoHistoryModule.cpp
--- a/Engine/Source/Editor/UndoHistory/Private/UndoHistoryModule.cpp
+++ b/Engine/Source/Editor/UndoHistory/Private/UndoHistoryModule.cpp
@@ -7,6 +7,8 @@
 
 #define LOCTEXT_NAMESPACE "FUndoHistoryModule"
 
+const FName UndoHistoryTabName("UndoHistory");
+
 void FUndoHistoryModule::StartupModule()
 {
 	FGlobalTabmanager::Get()->RegisterNomadTabSpawner(UndoHistoryTabName, FOnSpawnTab::CreateRaw(this, &FUndoHistoryModule::HandleSpawnSettingsTab))
diff --git a/Engine/Source/Editor/UndoHistory/Public/UndoHistoryModule.h b/Engine/Source/Editor/UndoHistory/Public/UndoHistoryModule.h
--- a/Engine/Source/Editor/UndoHistory/Public/UndoHistoryModule.h
+++ b/Engine/Source/Editor/UndoHistory/Public/UndoHistoryModule.h
@@ -6,7 +6,7 @@
 #include "Framework/Docking/TabManager.h"
 #include "Interfaces/IUndoHistoryModule.h"
 
-static const FName UndoHistoryTabName("UndoHistory");
+UNDOHISTORY_API extern const FName UndoHistoryTabName;
 
 /**
 * Implements the UndoHistory module.
