Fix for older devices that only support max thread count of 512

diff --git a/Engine/Plugins/Runtime/AR/AppleAR/AppleARKit/Source/AppleARKit/Private/AppleARKitTextures.cpp b/Engine/Plugins/Runtime/AR/AppleAR/AppleARKit/Source/AppleARKit/Private/AppleARKitTextures.cpp
--- a/Engine/Plugins/Runtime/AR/AppleAR/AppleARKit/Source/AppleARKit/Private/AppleARKitTextures.cpp
+++ b/Engine/Plugins/Runtime/AR/AppleAR/AppleARKit/Source/AppleARKit/Private/AppleARKitTextures.cpp
@@ -832,7 +832,7 @@ public:
 };
 #endif
 
-#define NUM_THREADS_PER_GROUP_DIMENSION 32
+#define NUM_THREADS_PER_GROUP_DIMENSION 16
 
 class FComputeShaderYCbCrToRGB : public FGlobalShader
 {
