Fix UE4 references in AppleARKitConversion.h

diff --git a/Engine/Plugins/Runtime/AR/AppleAR/AppleARKit/Source/AppleARKit/Private/AppleARKitMeshData.cpp b/Engine/Plugins/Runtime/AR/AppleAR/AppleARKit/Source/AppleARKit/Private/AppleARKitMeshData.cpp
--- a/Engine/Plugins/Runtime/AR/AppleAR/AppleARKit/Source/AppleARKit/Private/AppleARKitMeshData.cpp
+++ b/Engine/Plugins/Runtime/AR/AppleAR/AppleARKit/Source/AppleARKit/Private/AppleARKitMeshData.cpp
@@ -186,7 +186,7 @@ void FARKitMeshData::UpdateMeshData()
 				for (auto Index = 0; Index < NumVertices; ++Index, SourceVertices += 3)
 				{
 					// See FAppleARKitConversion::ToFVector
-					Vertices[Index] = FVector(-SourceVertices[2], SourceVertices[0], SourceVertices[1]) * FAppleARKitConversion::ToUE4Scale();
+					Vertices[Index] = FVector(-SourceVertices[2], SourceVertices[0], SourceVertices[1]) * FAppleARKitConversion::ToUEScale();
 					MergeBoxWith(BoundingBox, Vertices[Index]);
 				}
 			}
diff --git a/Engine/Plugins/Runtime/AR/AppleAR/AppleARKit/Source/AppleARKit/Public/AppleARKitConversion.h b/Engine/Plugins/Runtime/AR/AppleAR/AppleARKit/Source/AppleARKit/Public/AppleARKitConversion.h
--- a/Engine/Plugins/Runtime/AR/AppleAR/AppleARKit/Source/AppleARKit/Public/AppleARKitConversion.h
+++ b/Engine/Plugins/Runtime/AR/AppleAR/AppleARKit/Source/AppleARKit/Public/AppleARKitConversion.h
@@ -56,7 +56,7 @@ struct FARWorldSaveHeader
 
 struct APPLEARKIT_API FAppleARKitConversion
 {
-	static FORCEINLINE float ToUE4Scale()
+	static FORCEINLINE float ToUEScale()
 	{
 		return 100.f;
 	}
@@ -83,7 +83,7 @@ struct APPLEARKIT_API FAppleARKitConversion
 			FPlane(RawYUpMatrix.columns[3][0], RawYUpMatrix.columns[3][1], RawYUpMatrix.columns[3][2], RawYUpMatrix.columns[3][3]));
 
 		// Extract & convert translation
-		FVector Translation = FVector( -RawYUpFMatrix.M[3][2], RawYUpFMatrix.M[3][0], RawYUpFMatrix.M[3][1] ) * ToUE4Scale();
+		FVector Translation = FVector( -RawYUpFMatrix.M[3][2], RawYUpFMatrix.M[3][0], RawYUpFMatrix.M[3][1] ) * ToUEScale();
 
 		// Extract & convert rotation 
 		FQuat RawRotation( RawYUpFMatrix );
@@ -136,7 +136,7 @@ struct APPLEARKIT_API FAppleARKitConversion
 	 */
 	static FORCEINLINE FVector ToFVector(const vector_float3& RawYUpVector)
 	{
-		return FVector( -RawYUpVector.z, RawYUpVector.x, RawYUpVector.y ) * ToUE4Scale();
+		return FVector( -RawYUpVector.z, RawYUpVector.x, RawYUpVector.y ) * ToUEScale();
 	}
 
     /**
diff --git a/Engine/Plugins/Runtime/AR/AppleAR/AppleARKitFaceSupport/Source/AppleARKitFaceSupport/Private/AppleARKitFaceMeshComponent.cpp b/Engine/Plugins/Runtime/AR/AppleAR/AppleARKitFaceSupport/Source/AppleARKitFaceSupport/Private/AppleARKitFaceMeshComponent.cpp
--- a/Engine/Plugins/Runtime/AR/AppleAR/AppleARKitFaceSupport/Source/AppleARKitFaceSupport/Private/AppleARKitFaceMeshComponent.cpp
+++ b/Engine/Plugins/Runtime/AR/AppleAR/AppleARKitFaceSupport/Source/AppleARKitFaceSupport/Private/AppleARKitFaceMeshComponent.cpp
@@ -208,7 +208,7 @@ float UAppleARKitFaceMeshComponent::GetFaceBlendShapeAmount(EARFaceBlendShape Bl
 
 FMatrix UAppleARKitFaceMeshComponent::GetRenderMatrix() const
 {
-	const float Scale = FAppleARKitConversion::ToUE4Scale();
+	const float Scale = FAppleARKitConversion::ToUEScale();
 	
 	FTransform RenderTrans;
 	switch (TransformSetting)
