Fix Priori 1 iOS plugins issues on UE5 EA
When using WhitelistPlatforms on uplugin files the proper syntax for iOS is "IOS" not "iOS"

diff --git a/Engine/Plugins/Online/WebAuth/WebAuth.uplugin b/Engine/Plugins/Online/WebAuth/WebAuth.uplugin
--- a/Engine/Plugins/Online/WebAuth/WebAuth.uplugin
+++ b/Engine/Plugins/Online/WebAuth/WebAuth.uplugin
@@ -1,6 +1,6 @@
 {
 	"FileVersion" : 3,
-	
+
 	"FriendlyName" : "Web Authentication Plugin",
 	"Version" : 1,
 	"VersionName" : "1.0",
@@ -18,7 +18,7 @@
 			"WhitelistPlatforms" :
 			[
 				"Win64",
-				"iOS",
+				"IOS",
 				"Android"
 			],
 			"BlacklistTargets": [ "Server" ]
diff --git a/Engine/Plugins/Runtime/AR/AppleAR/AppleARKit/AppleARKit.uplugin b/Engine/Plugins/Runtime/AR/AppleAR/AppleARKit/AppleARKit.uplugin
--- a/Engine/Plugins/Runtime/AR/AppleAR/AppleARKit/AppleARKit.uplugin
+++ b/Engine/Plugins/Runtime/AR/AppleAR/AppleARKit/AppleARKit.uplugin
@@ -21,13 +21,13 @@
 			"Name": "AppleARKit",
 			"Type": "Runtime",
 			"LoadingPhase": "PostConfigInit",
-			"WhitelistPlatforms" : [ "Mac", "iOS", "Win64" ]
+			"WhitelistPlatforms" : [ "Mac", "IOS", "Win64" ]
 		},
 		{
 			"Name": "AppleARKitPoseTrackingLiveLink",
 			"Type": "Runtime",
 			"LoadingPhase": "PostConfigInit",
-			"WhitelistPlatforms" : [ "Mac", "iOS", "Win64" ]
+			"WhitelistPlatforms" : [ "Mac", "IOS", "Win64" ]
 
 		}
 	],
diff --git a/Engine/Plugins/Runtime/AR/AppleAR/AppleARKitFaceSupport/AppleARKitFaceSupport.uplugin b/Engine/Plugins/Runtime/AR/AppleAR/AppleARKitFaceSupport/AppleARKitFaceSupport.uplugin
--- a/Engine/Plugins/Runtime/AR/AppleAR/AppleARKitFaceSupport/AppleARKitFaceSupport.uplugin
+++ b/Engine/Plugins/Runtime/AR/AppleAR/AppleARKitFaceSupport/AppleARKitFaceSupport.uplugin
@@ -22,7 +22,7 @@
 			"Name": "AppleARKitFaceSupport",
 			"Type": "Runtime",
 			"LoadingPhase": "PostConfigInit",
-			"WhitelistPlatforms" : [ "Mac", "iOS", "Win64" ]
+			"WhitelistPlatforms" : [ "Mac", "IOS", "Win64" ]
 		}
 	],
     "Plugins":
diff --git a/Engine/Plugins/Runtime/AR/Google/GoogleARCoreServices/GoogleARCoreServices.uplugin b/Engine/Plugins/Runtime/AR/Google/GoogleARCoreServices/GoogleARCoreServices.uplugin
--- a/Engine/Plugins/Runtime/AR/Google/GoogleARCoreServices/GoogleARCoreServices.uplugin
+++ b/Engine/Plugins/Runtime/AR/Google/GoogleARCoreServices/GoogleARCoreServices.uplugin
@@ -16,7 +16,7 @@
         {
             "Name" : "GoogleARCoreServices",
             "Type" : "Runtime",
-            "WhitelistPlatforms": [ "Win64", "Mac", "Android", "Linux", "iOS"]
+            "WhitelistPlatforms": [ "Win64", "Mac", "Android", "Linux", "IOS"]
         },
     ],
     "Plugins": [
diff --git a/Engine/Plugins/Runtime/ResonanceAudio/ResonanceAudio.uplugin b/Engine/Plugins/Runtime/ResonanceAudio/ResonanceAudio.uplugin
--- a/Engine/Plugins/Runtime/ResonanceAudio/ResonanceAudio.uplugin
+++ b/Engine/Plugins/Runtime/ResonanceAudio/ResonanceAudio.uplugin
@@ -21,7 +21,7 @@
 		"LoadingPhase": "Default",
 		"WhitelistPlatforms": [
 			"Android",
-			"iOS",
+			"IOS",
 			"Linux",
 			"Mac",
 			"Win64",
diff --git a/Engine/Plugins/Runtime/WebBrowserNativeProxy/WebBrowserNativeProxy.uplugin b/Engine/Plugins/Runtime/WebBrowserNativeProxy/WebBrowserNativeProxy.uplugin
--- a/Engine/Plugins/Runtime/WebBrowserNativeProxy/WebBrowserNativeProxy.uplugin
+++ b/Engine/Plugins/Runtime/WebBrowserNativeProxy/WebBrowserNativeProxy.uplugin
@@ -21,11 +21,11 @@
 			"WhitelistPlatforms": [
 				"Win64",
 				"Mac",
-				"iOS",
+				"IOS",
 				"Android"
 			]
 		}
 	],
 	"CanContainContent": false,
 	"IsExperimentalVersion": false
-}
\ No newline at end of file
+}
