Fix UE4 reference in AppleARKitFaceSupport.h header (comment, "UE4" -> "UE").

diff --git a/Engine/Plugins/Runtime/AR/AppleAR/AppleARKit/Source/AppleARKit/Public/AppleARKitFaceSupport.h b/Engine/Plugins/Runtime/AR/AppleAR/AppleARKit/Source/AppleARKit/Public/AppleARKitFaceSupport.h
--- a/Engine/Plugins/Runtime/AR/AppleAR/AppleARKit/Source/AppleARKit/Public/AppleARKitFaceSupport.h
+++ b/Engine/Plugins/Runtime/AR/AppleAR/AppleARKit/Source/AppleARKit/Public/AppleARKitFaceSupport.h
@@ -49,7 +49,7 @@ public:
 	/**
 	 * Creates a face ar specific configuration object if that is requested without exposing the main code to the face APIs
 	 *
-	 * @param SessionConfig the UE4 configuration object that needs processing
+	 * @param SessionConfig the UE configuration object that needs processing
 	 * @param InProvider the custom timecode provider to use
 	 */
 	virtual ARConfiguration* ToARConfiguration(UARSessionConfig* SessionConfig, UTimecodeProvider* InProvider) { return nullptr; }
