Fix UE4 references in ARKit source (UE4 -> UE).

diff --git a/Engine/Plugins/Runtime/AR/AppleAR/AppleARKit/Source/AppleARKit/Private/AppleARKitFrame.h b/Engine/Plugins/Runtime/AR/AppleAR/AppleARKit/Source/AppleARKit/Private/AppleARKitFrame.h
--- a/Engine/Plugins/Runtime/AR/AppleAR/AppleARKit/Source/AppleARKit/Private/AppleARKitFrame.h
+++ b/Engine/Plugins/Runtime/AR/AppleAR/AppleARKit/Source/AppleARKit/Private/AppleARKitFrame.h
@@ -24,7 +24,7 @@ struct APPLEARKIT_API FAppleARKitFrame
 #if SUPPORTS_ARKIT_1_0
 	/** 
 	 * This is a conversion copy-constructor that takes a raw ARFrame and fills this structs members
-	 * with the UE4-ified versions of ARFrames properties.
+	 * with the UE-ified versions of ARFrames properties.
 	 *
 	 * @param InARFrame 		- The frame to convert / copy
 	 * @param MinCameraUV 		- The minimum (top left) UV used to render the passthrough camera
diff --git a/Engine/Plugins/Runtime/AR/AppleAR/AppleARKit/Source/AppleARKit/Private/AppleARKitLightEstimate.h b/Engine/Plugins/Runtime/AR/AppleAR/AppleARKit/Source/AppleARKit/Private/AppleARKitLightEstimate.h
--- a/Engine/Plugins/Runtime/AR/AppleAR/AppleARKit/Source/AppleARKit/Private/AppleARKitLightEstimate.h
+++ b/Engine/Plugins/Runtime/AR/AppleAR/AppleARKit/Source/AppleARKit/Private/AppleARKitLightEstimate.h
@@ -2,7 +2,7 @@
 
 #pragma once
 
-// UE4
+// UE
 #include "CoreMinimal.h"
 #include "Math/SHMath.h"
 #include "AppleARKitAvailability.h"
@@ -34,7 +34,7 @@ struct APPLEARKIT_API FAppleARKitLightEstimate
 
 	/** 
 	 * This is a conversion copy-constructor that takes a raw ARLightEstimate and fills this structs 
-	 * members with the UE4-ified versions of ARLightEstimate's properties.
+	 * members with the UE-ified versions of ARLightEstimate's properties.
 	 */ 
 	FAppleARKitLightEstimate( ARLightEstimate* InARLightEstimate );
 
