Fix iOS FVector2D cast

diff --git a/Engine/Plugins/Runtime/AR/AppleAR/AppleARKit/Source/AppleARKit/Private/AppleARKitTextures.cpp b/Engine/Plugins/Runtime/AR/AppleAR/AppleARKit/Source/AppleARKit/Private/AppleARKitTextures.cpp
--- a/Engine/Plugins/Runtime/AR/AppleAR/AppleARKit/Source/AppleARKit/Private/AppleARKitTextures.cpp
+++ b/Engine/Plugins/Runtime/AR/AppleAR/AppleARKit/Source/AppleARKit/Private/AppleARKitTextures.cpp
@@ -997,10 +997,10 @@ public:
 			FRDGBuilder GraphBuilder(RHICmdList);
 			FComputeShaderYCbCrToRGB::FParameters* PassParameters = GraphBuilder.AllocParameters<FComputeShaderYCbCrToRGB::FParameters>();
 			PassParameters->OutputTexture = DecodedTextureUAV;
-			PassParameters->OutputTextureSize = FVector2D(Size);
+			PassParameters->OutputTextureSize = FVector2f(Size);
 
 			PassParameters->InputTextureY = VideoTextureY;
-			PassParameters->InputTextureYSize = FVector2D(CapturedYImageSize);
+			PassParameters->InputTextureYSize = FVector2f(CapturedYImageSize);
 			
 			// This mapping must be the same as the comment above YCbCrToLinearRGB!
 			static const TMap<EDeviceScreenOrientation, int> DeviceOrientationIds =
@@ -1021,7 +1021,7 @@ public:
 			}
 			
 			PassParameters->InputTextureCbCr = VideoTextureCbCr;
-			PassParameters->InputTextureCbCrSize = FVector2D(CapturedCbCrImageSize);
+			PassParameters->InputTextureCbCrSize = FVector2f(CapturedCbCrImageSize);
 
 			TShaderMapRef<FComputeShaderYCbCrToRGB> ComputeShader(GetGlobalShaderMap(GMaxRHIFeatureLevel));
 			
