Fix for ARKit LCW fallout

#rnx

diff --git a/Engine/Plugins/Runtime/AR/AppleAR/AppleARKit/Source/AppleARKit/Private/AppleARKitMeshData.cpp b/Engine/Plugins/Runtime/AR/AppleAR/AppleARKit/Source/AppleARKit/Private/AppleARKitMeshData.cpp
--- a/Engine/Plugins/Runtime/AR/AppleAR/AppleARKit/Source/AppleARKit/Private/AppleARKitMeshData.cpp
+++ b/Engine/Plugins/Runtime/AR/AppleAR/AppleARKit/Source/AppleARKit/Private/AppleARKitMeshData.cpp
@@ -186,8 +186,8 @@ void FARKitMeshData::UpdateMeshData()
 				for (auto Index = 0; Index < NumVertices; ++Index, SourceVertices += 3)
 				{
 					// See FAppleARKitConversion::ToFVector
-					Vertices[Index] = FVector(-SourceVertices[2], SourceVertices[0], SourceVertices[1]) * FAppleARKitConversion::ToUEScale();
-					MergeBoxWith(BoundingBox, Vertices[Index]);
+					Vertices[Index] = FVector3f(-SourceVertices[2], SourceVertices[0], SourceVertices[1]) * FAppleARKitConversion::ToUEScale();
+					MergeBoxWith(BoundingBox, static_cast<FVector>(Vertices[Index]));
 				}
 			}
 			else
