Silence deprecation warnings for ARKit's hitTest:types: API

#rnx

diff --git a/Engine/Plugins/Runtime/AR/AppleAR/AppleARKit/Source/AppleARKit/Private/AppleARKitSystem.cpp b/Engine/Plugins/Runtime/AR/AppleAR/AppleARKit/Source/AppleARKit/Private/AppleARKitSystem.cpp
--- a/Engine/Plugins/Runtime/AR/AppleAR/AppleARKit/Source/AppleARKit/Private/AppleARKitSystem.cpp
+++ b/Engine/Plugins/Runtime/AR/AppleAR/AppleARKit/Source/AppleARKit/Private/AppleARKitSystem.cpp
@@ -721,7 +721,15 @@ TArray<FARTraceResult> FAppleARKitSystem::OnLineTraceTrackedObjects( const FVect
 				};
 				
 				// GEngine->AddOnScreenDebugMessage(-1, 5.f, FColor::Blue, FString::Printf(TEXT("Hit Test At Screen Position: x: %f, y: %f"), NormalizedImagePosition.X, NormalizedImagePosition.Y));
-				
+
+				//
+				// TODO: Re-enable deprecation warnings after updating the following code to use raycasting.
+				//
+				//   - 'ARHitTestResult' is deprecated: first deprecated in iOS 14.0 - Use raycasting
+				//   - 'hitTest:types:' is deprecated: first deprecated in iOS 14.0 - Use [ARSession raycast:]
+				//
+				PRAGMA_DISABLE_DEPRECATION_WARNINGS
+
 				// First run hit test against existing planes with extents (converting & filtering results as we go)
 				if (!!(TraceChannels & EARLineTraceChannels::PlaneUsingExtent) || !!(TraceChannels & EARLineTraceChannels::PlaneUsingBoundaryPolygon))
 				{
@@ -772,6 +780,8 @@ TArray<FARTraceResult> FAppleARKitSystem::OnLineTraceTrackedObjects( const FVect
 						}
 					}
 				}
+
+				PRAGMA_ENABLE_DEPRECATION_WARNINGS
 			}
 		}
 #endif
