Deprecating RHICreateTexture2D and RHICreateTextureExternal2D.

diff --git a/Engine/Plugins/Experimental/GPULightmass/Source/GPULightmass/Private/Scene/Scene.cpp b/Engine/Plugins/Experimental/GPULightmass/Source/GPULightmass/Private/Scene/Scene.cpp
--- a/Engine/Plugins/Experimental/GPULightmass/Source/GPULightmass/Private/Scene/Scene.cpp
+++ b/Engine/Plugins/Experimental/GPULightmass/Source/GPULightmass/Private/Scene/Scene.cpp
@@ -1756,8 +1756,11 @@ void CopyRectTiled(
 
 void ReadbackVolumetricLightmapDataLayerFromGPU(FRHICommandListImmediate& RHICmdList, FVolumetricLightmapDataLayer& Layer, FIntVector Dimensions)
 {
-	FRHIResourceCreateInfo CreateInfo(TEXT("VolumetricLightmapDataLayerReadback"));
-	FTexture2DRHIRef StagingTexture2DSlice = RHICreateTexture2D(Layer.Texture->GetSizeX(), Layer.Texture->GetSizeY(), Layer.Texture->GetFormat(), 1, 1, TexCreate_CPUReadback | TexCreate_HideInVisualizeTexture, CreateInfo);
+	const FRHITextureCreateDesc StagingDesc =
+		FRHITextureCreateDesc::Create2D(TEXT("VolumetricLightmapDataLayerReadback"), Layer.Texture->GetSizeX(), Layer.Texture->GetSizeY(), Layer.Texture->GetFormat())
+		.SetFlags(ETextureCreateFlags::CPUReadback | ETextureCreateFlags::HideInVisualizeTexture);
+
+	FTextureRHIRef StagingTexture2DSlice = RHICreateTexture(StagingDesc);
 	FGPUFenceRHIRef Fence = RHICreateGPUFence(TEXT("VolumetricLightmapDataLayerReadback"));
 
 	check(Dimensions.Z == Layer.Texture->GetSizeZ());
diff --git a/Engine/Plugins/Experimental/VirtualHeightfieldMesh/Source/VirtualHeightfieldMesh/Private/VirtualHeightfieldMeshSceneProxy.cpp b/Engine/Plugins/Experimental/VirtualHeightfieldMesh/Source/VirtualHeightfieldMesh/Private/VirtualHeightfieldMeshSceneProxy.cpp
--- a/Engine/Plugins/Experimental/VirtualHeightfieldMesh/Source/VirtualHeightfieldMesh/Private/VirtualHeightfieldMeshSceneProxy.cpp
+++ b/Engine/Plugins/Experimental/VirtualHeightfieldMesh/Source/VirtualHeightfieldMesh/Private/VirtualHeightfieldMeshSceneProxy.cpp
@@ -677,9 +677,13 @@ void FVirtualHeightfieldMeshSceneProxy::AcceptOcclusionResults(FSceneView const*
 		FOcclusionResults& OcclusionResults = GOcclusionResults.Emplace(FOcclusionResultsKey(this, View));
 		OcclusionResults.TextureSize = OcclusionGridSize;
 		OcclusionResults.NumTextureMips = NumOcclusionLods;
-		
-		FRHIResourceCreateInfo CreateInfo(TEXT("VirtualHeightfieldMesh.OcclusionTexture"));
-		OcclusionResults.OcclusionTexture = RHICreateTexture2D(OcclusionGridSize.X, OcclusionGridSize.Y, PF_G8, NumOcclusionLods, 1, TexCreate_ShaderResource, CreateInfo);
+
+		const FRHITextureCreateDesc Desc =
+			FRHITextureCreateDesc::Create2D(TEXT("VirtualHeightfieldMesh.OcclusionTexture"), OcclusionGridSize, PF_G8)
+			.SetNumMips(NumOcclusionLods)
+			.SetFlags(ETextureCreateFlags::ShaderResource);
+
+		OcclusionResults.OcclusionTexture = RHICreateTexture(Desc);
 		
 		bool const* Src = Results->GetData() + ResultsStart;
 		FIntPoint Size = OcclusionGridSize;
@@ -853,15 +857,17 @@ namespace VirtualHeightfieldMesh
 	public:
 		virtual void InitRHI() override
 		{
-			FRHIResourceCreateInfo CreateInfo(TEXT("VirtualHeightfieldMesh.MinMaxDefaultTexture"));
-			FTexture2DRHIRef Texture2D = RHICreateTexture2D(1, 1, PF_B8G8R8A8, 1, 1, TexCreate_ShaderResource, CreateInfo);
-			TextureRHI = Texture2D;
+			const FRHITextureCreateDesc Desc =
+				FRHITextureCreateDesc::Create2D(TEXT("VirtualHeightfieldMesh.MinMaxDefaultTexture"), 1, 1, PF_B8G8R8A8)
+				.SetFlags(ETextureCreateFlags::ShaderResource);
+
+			TextureRHI = RHICreateTexture(Desc);
 
 			// Write the contents of the texture.
 			uint32 DestStride;
-			FColor* DestBuffer = (FColor*)RHILockTexture2D(Texture2D, 0, RLM_WriteOnly, DestStride, false);
+			FColor* DestBuffer = (FColor*)RHILockTexture2D(TextureRHI, 0, RLM_WriteOnly, DestStride, false);
 			*DestBuffer = FColor(0, 0, 255, 255);
-			RHIUnlockTexture2D(Texture2D, 0, false);
+			RHIUnlockTexture2D(TextureRHI, 0, false);
 
 			// Create the sampler state RHI resource.
 			FSamplerStateInitializerRHI SamplerStateInitializer(SF_Point, AM_Clamp, AM_Clamp, AM_Clamp);
diff --git a/Engine/Plugins/Experimental/VirtualHeightfieldMesh/Source/VirtualHeightfieldMeshEditor/Private/HeightfieldMinMaxTextureBuild.cpp b/Engine/Plugins/Experimental/VirtualHeightfieldMesh/Source/VirtualHeightfieldMeshEditor/Private/HeightfieldMinMaxTextureBuild.cpp
--- a/Engine/Plugins/Experimental/VirtualHeightfieldMesh/Source/VirtualHeightfieldMeshEditor/Private/HeightfieldMinMaxTextureBuild.cpp
+++ b/Engine/Plugins/Experimental/VirtualHeightfieldMesh/Source/VirtualHeightfieldMeshEditor/Private/HeightfieldMinMaxTextureBuild.cpp
@@ -50,12 +50,16 @@ namespace
 				FIntPoint(NumTilesX, NumTilesY), PF_R8G8B8A8, FClearValueBinding::None, TexCreate_None, TexCreate_UAV | TexCreate_ShaderResource | TexCreate_GenerateMipCapable | TexCreate_RenderTargetable, false, NumMips);
 			GRenderTargetPool.FindFreeElement(RHICmdList, FinalRenderTargetDesc, FinalRenderTarget, TEXT("FinalTarget"));
 
-			FRHIResourceCreateInfo CreateInfo(TEXT("FMinMaxTileRenderResources_StagingTextures"));
 			for (int32 MipLevel = 0; MipLevel < NumMips; MipLevel++)
 			{
 				const int32 SizeX = FMath::Max(NumTilesX >> MipLevel, 1);
 				const int32 SizeY = FMath::Max(NumTilesY >> MipLevel, 1);
-				StagingTextures.Add(RHICreateTexture2D(SizeX, SizeY, PF_R8G8B8A8, 1, 1, TexCreate_CPUReadback, CreateInfo));
+
+				const FRHITextureCreateDesc Desc =
+					FRHITextureCreateDesc::Create2D(TEXT("FMinMaxTileRenderResources_StagingTextures"), SizeX, SizeY, PF_R8G8B8A8)
+					.SetFlags(ETextureCreateFlags::CPUReadback);
+
+				StagingTextures.Add(RHICreateTexture(Desc));
 			}
 			
 			Fence = RHICreateGPUFence(TEXT("Runtime Virtual Texture Build"));
@@ -78,7 +82,7 @@ namespace
 		int32 GetNumFinalTexels() const { return NumFinalTexels; }
 		TRefCountPtr<IPooledRenderTarget> GetTileRenderTarget() const { return TileRenderTarget; }
 		TRefCountPtr<IPooledRenderTarget> GetFinalRenderTarget() const { return FinalRenderTarget; }
-		FTexture2DRHIRef GetStagingTexture(int32 InMipLevel) const { return StagingTextures[InMipLevel]; }
+		FTextureRHIRef GetStagingTexture(int32 InMipLevel) const { return StagingTextures[InMipLevel]; }
 		FRHIGPUFence* GetFence() const { return Fence; }
 
 	private:
@@ -90,7 +94,7 @@ namespace
 
 		TRefCountPtr<IPooledRenderTarget> TileRenderTarget;
 		TRefCountPtr<IPooledRenderTarget> FinalRenderTarget;
-		TArray<FTexture2DRHIRef> StagingTextures;
+		TArray<FTextureRHIRef> StagingTextures;
 		FGPUFenceRHIRef Fence;
 	};
 }
diff --git a/Engine/Plugins/Media/AndroidCamera/Source/AndroidCamera/Private/Player/AndroidCameraPlayer.cpp b/Engine/Plugins/Media/AndroidCamera/Source/AndroidCamera/Private/Player/AndroidCameraPlayer.cpp
--- a/Engine/Plugins/Media/AndroidCamera/Source/AndroidCamera/Private/Player/AndroidCameraPlayer.cpp
+++ b/Engine/Plugins/Media/AndroidCamera/Source/AndroidCamera/Private/Player/AndroidCameraPlayer.cpp
@@ -315,13 +315,16 @@ static void DoUpdateExternalCameraSampleExecute(TWeakPtr<FJavaAndroidCameraPlaye
 	FTextureRHIRef VideoTexture = PinnedJavaCameraPlayer->GetVideoTexture();
 	if (VideoTexture == nullptr)
 	{
-		FRHIResourceCreateInfo CreateInfo(TEXT("VideoTexture"));
-		VideoTexture = RHICreateTextureExternal2D(1, 1, PF_R8G8B8A8, 1, 1, TexCreate_None, ERHIAccess::Unknown, CreateInfo);
+		const FRHITextureCreateDesc Desc =
+			FRHITextureCreateDesc::Create2D(TEXT("VideoTexture"), 1, 1, PF_R8G8B8A8)
+			.SetFlags(ETextureCreateFlags::External);
+
+		VideoTexture = RHICreateTexture(Desc);
 		PinnedJavaCameraPlayer->SetVideoTexture(VideoTexture);
 
 		if (VideoTexture == nullptr)
 		{
-			UE_LOG(LogAndroidCamera, Warning, TEXT("CreateTextureExternal2D failed!"));
+			UE_LOG(LogAndroidCamera, Warning, TEXT("RHICreateTexture failed!"));
 			return;
 		}
 
diff --git a/Engine/Plugins/Media/AndroidMedia/Source/AndroidMedia/Private/Player/AndroidMediaPlayer.cpp b/Engine/Plugins/Media/AndroidMedia/Source/AndroidMedia/Private/Player/AndroidMediaPlayer.cpp
--- a/Engine/Plugins/Media/AndroidMedia/Source/AndroidMedia/Private/Player/AndroidMediaPlayer.cpp
+++ b/Engine/Plugins/Media/AndroidMedia/Source/AndroidMedia/Private/Player/AndroidMediaPlayer.cpp
@@ -474,13 +474,16 @@ static void DoUpdateExternalMediaSampleExecute(TWeakPtr<FJavaAndroidMediaPlayer,
 	FTextureRHIRef VideoTexture = PinnedJavaMediaPlayer->GetVideoTexture();
 	if (VideoTexture == nullptr)
 	{
-		FRHIResourceCreateInfo CreateInfo(TEXT("VideoTexture"));
-		VideoTexture = RHICreateTextureExternal2D(1, 1, PF_R8G8B8A8, 1, 1, TexCreate_None, ERHIAccess::SRVGraphics, CreateInfo);
+		const FRHITextureCreateDesc Desc =
+			FRHITextureCreateDesc::Create2D(TEXT("VideoTexture"), 1, 1, PF_R8G8B8A8)
+			.SetFlags(ETextureCreateFlags::External);
+
+		VideoTexture = RHICreateTexture(Desc);
 		PinnedJavaMediaPlayer->SetVideoTexture(VideoTexture);
 
 		if (VideoTexture == nullptr)
 		{
-			UE_LOG(LogAndroidMedia, Warning, TEXT("CreateTextureExternal2D failed!"));
+			UE_LOG(LogAndroidMedia, Warning, TEXT("RHICreateTexture failed!"));
 			return;
 		}
 
diff --git a/Engine/Plugins/Media/AvfMedia/Source/AvfMedia/Private/Player/AvfMediaVideoSampler.cpp b/Engine/Plugins/Media/AvfMedia/Source/AvfMedia/Private/Player/AvfMediaVideoSampler.cpp
--- a/Engine/Plugins/Media/AvfMedia/Source/AvfMedia/Private/Player/AvfMediaVideoSampler.cpp
+++ b/Engine/Plugins/Media/AvfMedia/Source/AvfMedia/Private/Player/AvfMediaVideoSampler.cpp
@@ -252,9 +252,7 @@ void FAvfMediaVideoSampler::ProcessFrame(CVPixelBufferRef Frame, FTimespan Sampl
 		{
 			// Expecting BiPlanar kCVPixelFormatType_420YpCbCr8BiPlanar Full/Video
 			check(CVPixelBufferGetPlaneCount(Frame) == 2);
-			
-			ETextureCreateFlags TexCreateFlags = TexCreate_Dynamic | TexCreate_NoTiling;
-			
+
 			int32 YWidth = CVPixelBufferGetWidthOfPlane(Frame, 0);
 			int32 YHeight = CVPixelBufferGetHeightOfPlane(Frame, 0);
 			
@@ -272,19 +270,25 @@ void FAvfMediaVideoSampler::ProcessFrame(CVPixelBufferRef Frame, FTimespan Sampl
 			check(UVTextureRef);
 			
 			// Metal can upload directly from an IOSurface to a 2D texture, so we can just wrap it.
-			FRHIResourceCreateInfo YCreateInfo(TEXT("YTex"));
-			YCreateInfo.BulkData = new FAvfTexture2DResourceWrapper(YTextureRef);
-			YCreateInfo.ResourceArray = nullptr;
-			
-			FRHIResourceCreateInfo UVCreateInfo(TEXT("UVTex"));
-			UVCreateInfo.BulkData = new FAvfTexture2DResourceWrapper(UVTextureRef);
-			UVCreateInfo.ResourceArray = nullptr;
+
+			const FRHITextureCreateDesc YDesc =
+				FRHITextureCreateDesc::Create2D(TEXT("YTex"), YWidth, YHeight, PF_G8)
+				.SetFlags(ETextureCreateFlags::Dynamic | ETextureCreateFlags::NoTiling | ETextureCreateFlags::ShaderResource)
+				.SetBulkData(new FAvfTexture2DResourceWrapper(YTextureRef));
+
+			const FRHITextureCreateDesc UVDesc =
+				FRHITextureCreateDesc::Create2D(TEXT("UVTex"), UVWidth, UVHeight, PF_R8G8)
+				.SetFlags(ETextureCreateFlags::Dynamic | ETextureCreateFlags::NoTiling | ETextureCreateFlags::ShaderResource)
+				.SetBulkData(new FAvfTexture2DResourceWrapper(UVTextureRef));
 			
-			TRefCountPtr<FRHITexture2D> YTex = RHICreateTexture2D(YWidth, YHeight, PF_G8, 1, 1, TexCreateFlags | TexCreate_ShaderResource, YCreateInfo);
-			TRefCountPtr<FRHITexture2D> UVTex = RHICreateTexture2D(UVWidth, UVHeight, PF_R8G8, 1, 1, TexCreateFlags | TexCreate_ShaderResource, UVCreateInfo);
+			TRefCountPtr<FRHITexture> YTex = RHICreateTexture(YDesc);
+			TRefCountPtr<FRHITexture> UVTex = RHICreateTexture(UVDesc);
+
+			const FRHITextureCreateDesc Desc =
+				FRHITextureCreateDesc::Create2D(TEXT("Info"), YWidth, YHeight, PF_B8G8R8A8)
+				.SetFlags(ETextureCreateFlags::Dynamic | ETextureCreateFlags::NoTiling | ETextureCreateFlags::ShaderResource | ETextureCreateFlags::RenderTargetable | ETextureCreateFlags::SRGB);
 			
-			FRHIResourceCreateInfo Info(TEXT("Info"));
-			ShaderResource = RHICreateTexture2D(YWidth, YHeight, PF_B8G8R8A8, 1, 1, TexCreateFlags | TexCreate_ShaderResource | TexCreate_RenderTargetable | TexCreate_SRGB, Info);
+			ShaderResource = RHICreateTexture(Desc);
 			
 			// render video frame into sink texture
 			FRHICommandListImmediate& RHICmdList = FRHICommandListExecutor::GetImmediateCommandList();
@@ -335,15 +339,13 @@ void FAvfMediaVideoSampler::ProcessFrame(CVPixelBufferRef Frame, FTimespan Sampl
 			CVReturn Result = CVMetalTextureCacheCreateTextureFromImage(kCFAllocatorDefault, MetalTextureCache, Frame, nullptr, MTLPixelFormatBGRA8Unorm_sRGB, Width, Height, 0, &TextureRef);
 			check(Result == kCVReturnSuccess);
 			check(TextureRef);
-			
-			FRHIResourceCreateInfo CreateInfo(TEXT("FAvfMediaVideoSampler"));
-			CreateInfo.BulkData = new FAvfTexture2DResourceWrapper(TextureRef);
-			CreateInfo.ResourceArray = nullptr;
-			
-			ETextureCreateFlags TexCreateFlags = TexCreate_SRGB;
-			TexCreateFlags |= TexCreate_Dynamic | TexCreate_NoTiling;
-			
-			ShaderResource = RHICreateTexture2D(Width, Height, PF_B8G8R8A8, 1, 1, TexCreateFlags | TexCreate_ShaderResource, CreateInfo);
+
+			const FRHITextureCreateDesc Desc =
+				FRHITextureCreateDesc::Create2D(TEXT("FAvfMediaVideoSampler"), Width, Height, PF_B8G8R8A8)
+				.SetFlags(ETextureCreateFlags::SRGB | ETextureCreateFlags::Dynamic | ETextureCreateFlags::NoTiling | ETextureCreateFlags::ShaderResource)
+				.SetBulkData(new FAvfTexture2DResourceWrapper(TextureRef));
+
+			ShaderResource = RHICreateTexture(Desc);
 			
 			CFRelease(TextureRef);
 		}
diff --git a/Engine/Plugins/Media/BinkMedia/Source/BinkMediaPlayer/Private/binktexturesRHI.cpp b/Engine/Plugins/Media/BinkMedia/Source/BinkMediaPlayer/Private/binktexturesRHI.cpp
--- a/Engine/Plugins/Media/BinkMedia/Source/BinkMediaPlayer/Private/binktexturesRHI.cpp
+++ b/Engine/Plugins/Media/BinkMedia/Source/BinkMediaPlayer/Private/binktexturesRHI.cpp
@@ -57,11 +57,11 @@ typedef struct BINKTEXTURESRHI
 	S32 draw_flags;
 	F32 u0, v0, u1, v1;
 
-	FTexture2DRHIRef Ytexture[BINKMAXFRAMEBUFFERS];
-	FTexture2DRHIRef cRtexture[BINKMAXFRAMEBUFFERS];
-	FTexture2DRHIRef cBtexture[BINKMAXFRAMEBUFFERS];
-	FTexture2DRHIRef Atexture[BINKMAXFRAMEBUFFERS];
-	FTexture2DRHIRef Htexture[BINKMAXFRAMEBUFFERS];
+	FTextureRHIRef Ytexture[BINKMAXFRAMEBUFFERS];
+	FTextureRHIRef cRtexture[BINKMAXFRAMEBUFFERS];
+	FTextureRHIRef cBtexture[BINKMAXFRAMEBUFFERS];
+	FTextureRHIRef Atexture[BINKMAXFRAMEBUFFERS];
+	FTextureRHIRef Htexture[BINKMAXFRAMEBUFFERS];
 
 	// unused for now
 	S32 tonemap;
@@ -158,7 +158,14 @@ static BINKTEXTURES * Create_textures(BINKSHADERS * pshaders, HBINK bink, void *
 #else
 	ETextureCreateFlags TexCreateFlags = TexCreate_Dynamic | TexCreate_NoTiling;
 #endif
-	FRHIResourceCreateInfo CreateInfo(TEXT("BINK"));
+
+	const FRHITextureCreateDesc YADesc =
+		FRHITextureCreateDesc::Create2D(TEXT("BINK"), bb->YABufferWidth, bb->YABufferHeight, format)
+		.SetFlags(TexCreateFlags);
+
+	const FRHITextureCreateDesc cRcBDesc =
+		FRHITextureCreateDesc::Create2D(TEXT("BINK"), bb->cRcBBufferWidth, bb->cRcBBufferHeight, format)
+		.SetFlags(TexCreateFlags);
 
 	for (int i = 0; i < bb->TotalFrames; ++i)
 	{
@@ -166,7 +173,7 @@ static BINKTEXTURES * Create_textures(BINKSHADERS * pshaders, HBINK bink, void *
 
 		if (bp_src->YPlane.Allocate)
 		{
-			textures->Ytexture[i] = RHICreateTexture2D(bb->YABufferWidth, bb->YABufferHeight, format, 1, 1, TexCreateFlags, CreateInfo);
+			textures->Ytexture[i] = RHICreateTexture(YADesc);
 #if PLATFORM_HAS_DIRECT_TEXTURE_MEMORY_ACCESS 
 			bp_src->YPlane.Buffer = GDynamicRHI->LockTexture2D_RenderThread(RHICmdList, textures->Ytexture[i], 0, PLATFORM_DIRECT_TEXTURE_MEMORY_ACCESS_LOCK_MODE, bp_src->YPlane.BufferPitch, false, false);
 			bp_src->YPlane.BufferPitch = bp_src->YPlane.BufferPitch ? bp_src->YPlane.BufferPitch : ((bb->YABufferWidth+255)&-256);
@@ -175,7 +182,7 @@ static BINKTEXTURES * Create_textures(BINKSHADERS * pshaders, HBINK bink, void *
 
 		if (bp_src->cRPlane.Allocate)
 		{
-			textures->cRtexture[i] = RHICreateTexture2D(bb->cRcBBufferWidth, bb->cRcBBufferHeight, format, 1, 1, TexCreateFlags, CreateInfo);
+			textures->cRtexture[i] = RHICreateTexture(cRcBDesc);
 #if PLATFORM_HAS_DIRECT_TEXTURE_MEMORY_ACCESS 
 			bp_src->cRPlane.Buffer = GDynamicRHI->LockTexture2D_RenderThread(RHICmdList, textures->cRtexture[i], 0, PLATFORM_DIRECT_TEXTURE_MEMORY_ACCESS_LOCK_MODE, bp_src->cRPlane.BufferPitch, false, false);
 			bp_src->cRPlane.BufferPitch = bp_src->cRPlane.BufferPitch ? bp_src->cRPlane.BufferPitch : ((bb->cRcBBufferWidth+255)&-256);
@@ -184,7 +191,7 @@ static BINKTEXTURES * Create_textures(BINKSHADERS * pshaders, HBINK bink, void *
 
 		if (bp_src->cBPlane.Allocate)
 		{
-			textures->cBtexture[i] = RHICreateTexture2D(bb->cRcBBufferWidth, bb->cRcBBufferHeight, format, 1, 1, TexCreateFlags, CreateInfo);
+			textures->cBtexture[i] = RHICreateTexture(cRcBDesc);
 #if PLATFORM_HAS_DIRECT_TEXTURE_MEMORY_ACCESS 
 			bp_src->cBPlane.Buffer = GDynamicRHI->LockTexture2D_RenderThread(RHICmdList, textures->cBtexture[i], 0, PLATFORM_DIRECT_TEXTURE_MEMORY_ACCESS_LOCK_MODE, bp_src->cBPlane.BufferPitch, false, false);
 			bp_src->cBPlane.BufferPitch = bp_src->cBPlane.BufferPitch ? bp_src->cBPlane.BufferPitch : ((bb->cRcBBufferWidth+255)&-256);
@@ -193,7 +200,7 @@ static BINKTEXTURES * Create_textures(BINKSHADERS * pshaders, HBINK bink, void *
 
 		if (bp_src->APlane.Allocate)
 		{
-			textures->Atexture[i] = RHICreateTexture2D(bb->YABufferWidth, bb->YABufferHeight, format, 1, 1, TexCreateFlags, CreateInfo);
+			textures->Atexture[i] = RHICreateTexture(YADesc);
 #if PLATFORM_HAS_DIRECT_TEXTURE_MEMORY_ACCESS 
 			bp_src->APlane.Buffer = GDynamicRHI->LockTexture2D_RenderThread(RHICmdList, textures->Atexture[i], 0, PLATFORM_DIRECT_TEXTURE_MEMORY_ACCESS_LOCK_MODE, bp_src->APlane.BufferPitch, false, false);
 			bp_src->APlane.BufferPitch = bp_src->APlane.BufferPitch ? bp_src->APlane.BufferPitch : ((bb->YABufferWidth+255)&-256);
@@ -202,7 +209,7 @@ static BINKTEXTURES * Create_textures(BINKSHADERS * pshaders, HBINK bink, void *
 
 		if (bp_src->HPlane.Allocate)
 		{
-			textures->Htexture[i] = RHICreateTexture2D(bb->YABufferWidth, bb->YABufferHeight, format, 1, 1, TexCreateFlags, CreateInfo);
+			textures->Htexture[i] = RHICreateTexture(YADesc);
 #if PLATFORM_HAS_DIRECT_TEXTURE_MEMORY_ACCESS 
 			bp_src->HPlane.Buffer = GDynamicRHI->LockTexture2D_RenderThread(RHICmdList, textures->Htexture[i], 0, PLATFORM_DIRECT_TEXTURE_MEMORY_ACCESS_LOCK_MODE, bp_src->HPlane.BufferPitch, false, false);
 			bp_src->HPlane.BufferPitch = bp_src->HPlane.BufferPitch ? bp_src->HPlane.BufferPitch : ((bb->YABufferWidth+255)&-256);
diff --git a/Engine/Plugins/Media/ElectraUtil/Source/ElectraSamples/Private/ElectraTextureSampleAVF.cpp b/Engine/Plugins/Media/ElectraUtil/Source/ElectraSamples/Private/ElectraTextureSampleAVF.cpp
--- a/Engine/Plugins/Media/ElectraUtil/Source/ElectraSamples/Private/ElectraTextureSampleAVF.cpp
+++ b/Engine/Plugins/Media/ElectraUtil/Source/ElectraSamples/Private/ElectraTextureSampleAVF.cpp
@@ -209,17 +209,19 @@ void FElectraMediaTexConvApple::ConvertTexture(FTexture2DRHIRef & InDstTexture,
 			check(UVTextureRef);
 
 			// Metal can upload directly from an IOSurface to a 2D texture, so we can just wrap it.
-			FRHIResourceCreateInfo YCreateInfo(TEXT("YTex"));
-			YCreateInfo.BulkData = new FTexConvTexResourceWrapper(YTextureRef);
-			YCreateInfo.ResourceArray = nullptr;
 
-			FRHIResourceCreateInfo UVCreateInfo(TEXT("UVTex"));
-			UVCreateInfo.BulkData = new FTexConvTexResourceWrapper(UVTextureRef);
-			UVCreateInfo.ResourceArray = nullptr;
+			const FRHITextureCreateDesc YDesc =
+				FRHITextureCreateDesc::Create2D(TEXT("YTex"), YWidth, YHeight, PF_G8)
+				.SetFlags(ETextureCreateFlags::Dynamic | ETextureCreateFlags::NoTiling | ETextureCreateFlags::ShaderResource)
+				.SetBulkData(new FTexConvTexResourceWrapper(YTextureRef));
 
-			ETextureCreateFlags TexCreateFlags = TexCreate_Dynamic | TexCreate_NoTiling | TexCreate_ShaderResource;
-			TRefCountPtr<FRHITexture2D> YTex = RHICreateTexture2D(YWidth, YHeight, PF_G8, 1, 1, TexCreateFlags, YCreateInfo);
-			TRefCountPtr<FRHITexture2D> UVTex = RHICreateTexture2D(UVWidth, UVHeight, PF_R8G8, 1, 1, TexCreateFlags, UVCreateInfo);
+			const FRHITextureCreateDesc UVDesc =
+				FRHITextureCreateDesc::Create2D(TEXT("UVTex"), UVWidth, UVHeight, PF_R8G8)
+				.SetFlags(ETextureCreateFlags::Dynamic | ETextureCreateFlags::NoTiling | ETextureCreateFlags::ShaderResource)
+				.SetBulkData(new FTexConvTexResourceWrapper(UVTextureRef));
+
+			TRefCountPtr<FRHITexture> YTex = RHICreateTexture(YDesc);
+			TRefCountPtr<FRHITexture> UVTex = RHICreateTexture(UVDesc);
 
 			// render video frame into sink texture
 			FRHICommandListImmediate& RHICmdList = FRHICommandListExecutor::GetImmediateCommandList();
@@ -274,12 +276,12 @@ void FElectraMediaTexConvApple::ConvertTexture(FTexture2DRHIRef & InDstTexture,
 			check(Result == kCVReturnSuccess);
 			check(TextureRef);
 
-			FRHIResourceCreateInfo CreateInfo(TEXT("DstTexture"));
-			CreateInfo.BulkData = new FTexConvTexResourceWrapper(TextureRef);
-			CreateInfo.ResourceArray = nullptr;
+			const FRHITextureCreateDesc Desc =
+				FRHITextureCreateDesc::Create2D(TEXT("DstTexture"), Width, Height, PF_B8G8R8A8)
+				.SetFlags(ETextureCreateFlags::SRGB | ETextureCreateFlags::Dynamic | ETextureCreateFlags::NoTiling | ETextureCreateFlags::ShaderResource)
+				.SetBulkData(new FTexConvTexResourceWrapper(TextureRef));
 
-			ETextureCreateFlags TexCreateFlags = TexCreate_SRGB | TexCreate_Dynamic | TexCreate_NoTiling | TexCreate_ShaderResource;
-			InDstTexture = RHICreateTexture2D(Width, Height, PF_B8G8R8A8, 1, 1, TexCreateFlags, CreateInfo);
+			InDstTexture = RHICreateTexture(Desc);
 
 			CFRelease(TextureRef);
 		}
diff --git a/Engine/Plugins/Media/ElectraUtil/Source/ElectraSamples/Private/ElectraTextureSamplePC.cpp b/Engine/Plugins/Media/ElectraUtil/Source/ElectraSamples/Private/ElectraTextureSamplePC.cpp
--- a/Engine/Plugins/Media/ElectraUtil/Source/ElectraSamples/Private/ElectraTextureSamplePC.cpp
+++ b/Engine/Plugins/Media/ElectraUtil/Source/ElectraSamples/Private/ElectraTextureSamplePC.cpp
@@ -174,7 +174,7 @@ bool FElectraTextureSample::Convert(FTexture2DRHIRef & InDstTexture, const FConv
 	FIntPoint Dim = VideoDecoderOutput->GetDim();
 
 	bool bCrossDeviceCopy;
-	uint8 Format;
+	EPixelFormat Format;
 	if (VideoDecoderOutput->GetOutputType() != FVideoDecoderOutputPC::EOutputType::HardwareWin8Plus)
 	{
 		//
@@ -200,16 +200,11 @@ bool FElectraTextureSample::Convert(FTexture2DRHIRef & InDstTexture, const FConv
 	// Do we need a new RHI texture?
 	if (!Texture.IsValid() || Texture->GetSizeX() != Dim.X || Texture->GetSizeY() != Dim.Y)
 	{
-		FRHIResourceCreateInfo CreateInfo(TEXT("FElectraTextureSample"));
-		const ETextureCreateFlags CreateFlags = TexCreate_Dynamic;
-		Texture = RHICreateTexture2D(
-			Dim.X,
-			Dim.Y,
-			Format,
-			1,
-			1,
-			CreateFlags,
-			CreateInfo);
+		const FRHITextureCreateDesc Desc =
+			FRHITextureCreateDesc::Create2D(TEXT("FElectraTextureSample"), Dim, Format)
+			.SetFlags(ETextureCreateFlags::Dynamic);
+
+		Texture = RHICreateTexture(Desc);
 	}
 	
 	// Copy data into RHI texture
diff --git a/Engine/Plugins/Media/ElectraUtil/Source/ElectraSamples/Public/ElectraSamplesTextureManager.h b/Engine/Plugins/Media/ElectraUtil/Source/ElectraSamples/Public/ElectraSamplesTextureManager.h
--- a/Engine/Plugins/Media/ElectraUtil/Source/ElectraSamples/Public/ElectraSamplesTextureManager.h
+++ b/Engine/Plugins/Media/ElectraUtil/Source/ElectraSamples/Public/ElectraSamplesTextureManager.h
@@ -13,11 +13,11 @@ public:
 	virtual ~FElectraSamplesTextureManager();
 
 	TSharedPtr<IVideoDecoderTexture, ESPMode::ThreadSafe> CreateTexture(const FIntPoint& Dim, EPixelFormat Fmt = PF_Unknown);
-	FTexture2DRHIRef GetRHITextureFromPlatformTexture(void* PlatformTexture);
+	FTextureRHIRef GetRHITextureFromPlatformTexture(void* PlatformTexture);
 
 protected:
 	virtual TSharedPtr<IVideoDecoderTexture, ESPMode::ThreadSafe> PlatformCreateTexture(const FIntPoint& Dim, EPixelFormat Fmt) = 0;
-	virtual FTexture2DRHIRef PlatformGetTextureRHI(IVideoDecoderTexture *Texture) = 0;
+	virtual FTextureRHIRef PlatformGetTextureRHI(IVideoDecoderTexture *Texture) = 0;
 	virtual void* PlatformGetTexturePlatform(IVideoDecoderTexture* Texture) = 0;
 	
 	void PlatformAddSharedTextureRef(TSharedPtr<IVideoDecoderTexture, ESPMode::ThreadSafe> InTextureRef);
diff --git a/Engine/Plugins/Media/MediaIOFramework/Source/MediaIOCore/Private/MediaCapture.cpp b/Engine/Plugins/Media/MediaIOFramework/Source/MediaIOCore/Private/MediaCapture.cpp
--- a/Engine/Plugins/Media/MediaIOFramework/Source/MediaIOCore/Private/MediaCapture.cpp
+++ b/Engine/Plugins/Media/MediaIOFramework/Source/MediaIOCore/Private/MediaCapture.cpp
@@ -1170,7 +1170,6 @@ void UMediaCapture::InitializeOutputResources(int32 InNumberOfBuffers)
 	ENQUEUE_RENDER_COMMAND(MediaOutputCaptureFrameCreateResources)(
 		[This](FRHICommandListImmediate& RHICmdList)
 		{
-			FRHIResourceCreateInfo CreateInfo(TEXT("UMediaCapture"));
 			for (int32 Index = 0; Index < This->NumberOfCaptureFrame; ++Index)
 			{
 				if (This->DesiredOutputResourceType == EMediaCaptureResourceType::Texture)
diff --git a/Engine/Plugins/Media/WebMMedia/Source/WebMMedia/Private/Player/WebMVideoDecoder.cpp b/Engine/Plugins/Media/WebMMedia/Source/WebMMedia/Private/Player/WebMVideoDecoder.cpp
--- a/Engine/Plugins/Media/WebMMedia/Source/WebMMedia/Private/Player/WebMVideoDecoder.cpp
+++ b/Engine/Plugins/Media/WebMMedia/Source/WebMMedia/Private/Player/WebMVideoDecoder.cpp
@@ -169,14 +169,21 @@ void FWebMVideoDecoder::DoDecodeVideoFrames(const TArray<TSharedPtr<FWebMFrame>>
 
 void FWebMVideoDecoder::CreateTextures(const vpx_image_t* Image)
 {
-	FRHIResourceCreateInfo CreateInfo(TEXT("FWebMVideoDecoder_DecodedY"));
-	DecodedY = RHICreateTexture2D(Image->stride[0], Image->d_h, PF_G8, 1, 1, TexCreate_Dynamic, CreateInfo);
+	const FRHITextureCreateDesc DecodedYDesc =
+		FRHITextureCreateDesc::Create2D(TEXT("FWebMVideoDecoder_DecodedY"), Image->stride[0], Image->d_h, PF_G8)
+		.SetFlags(ETextureCreateFlags::Dynamic);
 
-	CreateInfo.DebugName = TEXT("FWebMVideoDecoder_DecodedU");
-	DecodedU = RHICreateTexture2D(Image->stride[1], Image->d_h / 2, PF_G8, 1, 1, TexCreate_Dynamic, CreateInfo);
+	const FRHITextureCreateDesc DecodedUDesc =
+		FRHITextureCreateDesc::Create2D(TEXT("FWebMVideoDecoder_DecodedU"), Image->stride[1], Image->d_h / 2, PF_G8)
+		.SetFlags(ETextureCreateFlags::Dynamic);
 
-	CreateInfo.DebugName = TEXT("FWebMVideoDecoder_DecodedV");
-	DecodedV = RHICreateTexture2D(Image->stride[2], Image->d_h / 2, PF_G8, 1, 1, TexCreate_Dynamic, CreateInfo);
+	const FRHITextureCreateDesc DecodedVDesc =
+		FRHITextureCreateDesc::Create2D(TEXT("FWebMVideoDecoder_DecodedV"), Image->stride[2], Image->d_h / 2, PF_G8)
+		.SetFlags(ETextureCreateFlags::Dynamic);
+
+	DecodedY = RHICreateTexture(DecodedYDesc);
+	DecodedU = RHICreateTexture(DecodedUDesc);
+	DecodedV = RHICreateTexture(DecodedVDesc);
 }
 
 void FWebMVideoDecoder::Close()
diff --git a/Engine/Plugins/Media/WmfMedia/Source/WmfMedia/Public/WmfMediaHardwareVideoDecodingTextureSample.h b/Engine/Plugins/Media/WmfMedia/Source/WmfMedia/Public/WmfMediaHardwareVideoDecodingTextureSample.h
--- a/Engine/Plugins/Media/WmfMedia/Source/WmfMedia/Public/WmfMediaHardwareVideoDecodingTextureSample.h
+++ b/Engine/Plugins/Media/WmfMedia/Source/WmfMedia/Public/WmfMediaHardwareVideoDecodingTextureSample.h
@@ -114,29 +114,28 @@ public:
 	 *
 	 * @return Destination texture 
 	 */
-	FTexture2DRHIRef GetOrCreateDestinationTexture()
+	FTextureRHIRef GetOrCreateDestinationTexture()
 	{
 		if (DestinationTexture.IsValid() && DestinationTexture->GetSizeX() == Dim.X && DestinationTexture->GetSizeY() == Dim.Y)
 		{
 			ETextureCreateFlags CurrentFlags = DestinationTexture->GetFlags();
-			bool bIsCurrentSRGB = EnumHasAnyFlags(CurrentFlags, TexCreate_SRGB);
+			bool bIsCurrentSRGB = EnumHasAnyFlags(CurrentFlags, ETextureCreateFlags::SRGB);
 			if (bIsCurrentSRGB == bIsDestinationTextureSRGB)
 			{
 				return DestinationTexture;
 			}
 		}
-		
-		FRHIResourceCreateInfo CreateInfo(TEXT("FWmfMediaHardwareVideoDecodingTextureSample_DestinationTexture"));
-		const ETextureCreateFlags CreateFlags = TexCreate_Dynamic | TexCreate_DisableSRVCreation |
-			(bIsDestinationTextureSRGB ? TexCreate_SRGB : TexCreate_None);
-		DestinationTexture = RHICreateTexture2D(
-			Dim.X,
-			Dim.Y,
-			Format,
-			1,
-			1,
-			CreateFlags,
-			CreateInfo);
+
+		FRHITextureCreateDesc Desc =
+			FRHITextureCreateDesc::Create2D(TEXT("FWmfMediaHardwareVideoDecodingTextureSample_DestinationTexture"), Dim, Format)
+			.SetFlags(ETextureCreateFlags::Dynamic | ETextureCreateFlags::DisableSRVCreation);
+
+		if (bIsDestinationTextureSRGB)
+		{
+			Desc.AddFlags(ETextureCreateFlags::SRGB);
+		}
+
+		DestinationTexture = RHICreateTexture(Desc);
 
 		return DestinationTexture;
 	}
@@ -146,23 +145,18 @@ public:
 	 *
 	 * @return Destination texture
 	 */
-	FTexture2DRHIRef GetOrCreateDestinationAlphaTexture()
+	FTextureRHIRef GetOrCreateDestinationAlphaTexture()
 	{
 		if (DestinationAlphaTexture.IsValid() && DestinationAlphaTexture->GetSizeX() == Dim.X && DestinationAlphaTexture->GetSizeY() == Dim.Y)
 		{
 			return DestinationAlphaTexture;
 		}
+
+		const FRHITextureCreateDesc Desc =
+			FRHITextureCreateDesc::Create2D(TEXT("FWmfMediaHardwareVideoDecodingTextureSample_DestinationAlphaTexture"), Dim, AlphaFormat)
+			.SetFlags(ETextureCreateFlags::Dynamic | ETextureCreateFlags::DisableSRVCreation);
 		
-		FRHIResourceCreateInfo CreateInfo(TEXT("FWmfMediaHardwareVideoDecodingTextureSample_DestinationAlphaTexture"));
-		const ETextureCreateFlags CreateFlags = TexCreate_Dynamic | TexCreate_DisableSRVCreation;
-		DestinationAlphaTexture = RHICreateTexture2D(
-			Dim.X,
-			Dim.Y,
-			AlphaFormat,
-			1,
-			1,
-			CreateFlags,
-			CreateInfo);
+		DestinationAlphaTexture = RHICreateTexture(Desc);
 
 		return DestinationAlphaTexture;
 	}
@@ -184,10 +178,10 @@ private:
 	TRefCountPtr<ID3D11Device> D3D11Device;
 
 	/** Destination Texture resource (from Rendering device) */
-	FTexture2DRHIRef DestinationTexture;
+	FTextureRHIRef DestinationTexture;
 
 	/** Destination Texture resource (from Rendering device) */
-	FTexture2DRHIRef DestinationAlphaTexture;
+	FTextureRHIRef DestinationAlphaTexture;
 
 	/** Texture format */
 	EPixelFormat Format;
diff --git a/Engine/Plugins/MovieScene/MovieRenderPipeline/Source/MovieRenderPipelineCore/Private/MoviePipelineSurfaceReader.cpp b/Engine/Plugins/MovieScene/MovieRenderPipeline/Source/MovieRenderPipelineCore/Private/MoviePipelineSurfaceReader.cpp
--- a/Engine/Plugins/MovieScene/MovieRenderPipeline/Source/MovieRenderPipelineCore/Private/MoviePipelineSurfaceReader.cpp
+++ b/Engine/Plugins/MovieScene/MovieRenderPipeline/Source/MovieRenderPipelineCore/Private/MoviePipelineSurfaceReader.cpp
@@ -50,17 +50,10 @@ void FMoviePipelineSurfaceReader::Resize(uint32 Width, uint32 Height)
 	ENQUEUE_RENDER_COMMAND(CreateCaptureFrameTexture)(
 		[Width, Height, This](FRHICommandListImmediate& RHICmdList)
 		{
-			FRHIResourceCreateInfo CreateInfo(TEXT("FMoviePipelineSurfaceReader"));
-
-			This->ReadbackTexture = RHICreateTexture2D(
-				Width,
-				Height,
-				This->PixelFormat,
-				1,
-				1,
-				TexCreate_CPUReadback,
-				CreateInfo
-				);
+			const FRHITextureCreateDesc Desc =
+				FRHITextureCreateDesc::Create2D(TEXT("FMoviePipelineSurfaceReader"), Width, Height, This->PixelFormat)
+				.SetFlags(ETextureCreateFlags::CPUReadback);
+			This->ReadbackTexture = RHICreateTexture(Desc);
 		});
 }
 
diff --git a/Engine/Plugins/Runtime/AR/AppleAR/AppleARKit/Source/AppleARKit/Private/AppleARKitTextures.cpp b/Engine/Plugins/Runtime/AR/AppleAR/AppleARKit/Source/AppleARKit/Private/AppleARKitTextures.cpp
--- a/Engine/Plugins/Runtime/AR/AppleAR/AppleARKit/Source/AppleARKit/Private/AppleARKitTextures.cpp
+++ b/Engine/Plugins/Runtime/AR/AppleAR/AppleARKit/Source/AppleARKit/Private/AppleARKitTextures.cpp
@@ -147,11 +147,15 @@ public:
 	{
 		FSamplerStateInitializerRHI SamplerStateInitializer(SF_Bilinear, AM_Clamp, AM_Clamp, AM_Clamp);
 		SamplerStateRHI = RHICreateSamplerState(SamplerStateInitializer);
-		
+
 		// Default to an empty 1x1 texture if we don't have a camera image
-		FRHIResourceCreateInfo CreateInfo(TEXT("FARKitTextureResource"));
 		Size.X = Size.Y = 1;
-		TextureRHI = RHICreateTexture2D(Size.X, Size.Y, PF_B8G8R8A8, 1, 1, TexCreate_ShaderResource, CreateInfo);
+
+		const FRHITextureCreateDesc Desc =
+			FRHITextureCreateDesc::Create2D(TEXT("FARKitTextureResource"), Size, PF_B8G8R8A8)
+			.SetFlags(ETextureCreateFlags::ShaderResource);
+		
+		TextureRHI = RHICreateTexture(Desc);
 		
 		OnTextureUpdated();
 	}
@@ -207,8 +211,11 @@ protected:
 			TextureRHI.SafeRelease();
 			
 			// Create the target texture that we'll update into
-			FRHIResourceCreateInfo CreateInfo(TEXT("FARKitTextureResource"));
-			TextureRHI = RHICreateTexture2D(Size.X, Size.Y, PixelFormat, 1, 1, TexCreate_Dynamic | TexCreate_ShaderResource | TexCreate_UAV, CreateInfo);
+			const FRHITextureCreateDesc Desc =
+				FRHITextureCreateDesc::Create2D(TEXT("FARKitTextureResource"), Size, PixelFormat)
+				.SetFlags(ETextureCreateFlags::Dynamic | ETextureCreateFlags::ShaderResource | ETextureCreateFlags::UAV);
+
+			TextureRHI = RHICreateTexture(Desc);
 		}
 	}
 	
@@ -890,11 +897,14 @@ public:
 		SamplerStateRHI = RHICreateSamplerState(SamplerStateInitializer);
 		
 		// Default to an empty 1x1 texture if we don't have a camera image
-		FRHIResourceCreateInfo CreateInfo(TEXT("DecodedTextureRef"));
 		Size.X = Size.Y = 1;
 		{
+			const FRHITextureCreateDesc Desc =
+				FRHITextureCreateDesc::Create2D(TEXT("DecodedTextureRef"), Size, PF_B8G8R8A8)
+				.SetFlags(ETextureCreateFlags::ShaderResource);
+
 			FScopeLock ScopeLock(&DecodedTextureLock);
-			DecodedTextureRef = RHICreateTexture2D(Size.X, Size.Y, PF_B8G8R8A8, 1, 1, TexCreate_ShaderResource, CreateInfo);
+			DecodedTextureRef = RHICreateTexture(Desc);
 		}
 		
 		UpdateTextureRHI();
@@ -943,16 +953,19 @@ public:
 		
 		{
 			// Update the RHI texture wrapper for the Y and CbCr images
-			FRHIResourceCreateInfo CreateInfo(TEXT("VideoTexture"));
-			const ETextureCreateFlags CreateFlags = TexCreate_Dynamic | TexCreate_NoTiling | TexCreate_ShaderResource;
-			CreateInfo.BulkData = new FAppleARKitCameraTextureResourceWrapper(CapturedYImage);
-			CreateInfo.ResourceArray = nullptr;
+			const FRHITextureCreateDesc YDesc =
+				FRHITextureCreateDesc::Create2D(TEXT("VideoTextureY"), CapturedYImageSize, PF_B8G8R8A8)
+				.SetFlags(ETextureCreateFlags::Dynamic | ETextureCreateFlags::NoTiling | ETextureCreateFlags::ShaderResource)
+				.SetBulkData(new FAppleARKitCameraTextureResourceWrapper(CapturedYImage));
 
-			// pull the Y and CbCr textures out of the captured image planes (format is fake here, it will get the format from the FAppleARKitCameraTextureResourceWrapper)
-			VideoTextureY = RHICreateTexture2D(CapturedYImageSize.X, CapturedYImageSize.Y, /*Format=*/PF_B8G8R8A8, /*NumMips=*/1, /*NumSamples=*/1, CreateFlags, CreateInfo);
+			const FRHITextureCreateDesc cBcRDesc =
+				FRHITextureCreateDesc::Create2D(TEXT("VideoTextureCbCr"), CapturedCbCrImageSize, PF_B8G8R8A8)
+				.SetFlags(ETextureCreateFlags::Dynamic | ETextureCreateFlags::NoTiling | ETextureCreateFlags::ShaderResource)
+				.SetBulkData(new FAppleARKitCameraTextureResourceWrapper(CapturedCbCrImage));
 
-			CreateInfo.BulkData = new FAppleARKitCameraTextureResourceWrapper(CapturedCbCrImage);
-			VideoTextureCbCr = RHICreateTexture2D(CapturedCbCrImageSize.X, CapturedCbCrImageSize.Y, /*Format=*/PF_B8G8R8A8, /*NumMips=*/1, /*NumSamples=*/1, CreateFlags, CreateInfo);
+			// pull the Y and CbCr textures out of the captured image planes (format is fake here, it will get the format from the FAppleARKitCameraTextureResourceWrapper)
+			VideoTextureY = RHICreateTexture(YDesc);
+			VideoTextureCbCr = RHICreateTexture(cBcRDesc);
 
 			// todo: Add an update call to the registry instead of this unregister/re-register
 			FExternalTextureRegistry::Get().UnregisterExternalTexture(ARKitPassthroughCameraExternalTextureYGuid);
@@ -984,10 +997,13 @@ public:
 		if (!DecodedTextureRef || Size != OutputSize)
 		{
 			Size = OutputSize;
-			FRHIResourceCreateInfo CreateInfo(TEXT("FARKitCameraVideoResource_DecodedTexture"));
 			{
+				const FRHITextureCreateDesc Desc =
+					FRHITextureCreateDesc::Create2D(TEXT("FARKitCameraVideoResource_DecodedTexture"), Size, PF_B8G8R8A8)
+					.SetFlags(ETextureCreateFlags::Dynamic | ETextureCreateFlags::ShaderResource | ETextureCreateFlags::UAV);
+
 				FScopeLock ScopeLock(&DecodedTextureLock);
-				DecodedTextureRef = RHICreateTexture2D(Size.X, Size.Y, PF_B8G8R8A8, 1, 1, TexCreate_Dynamic | TexCreate_ShaderResource | TexCreate_UAV, CreateInfo);
+				DecodedTextureRef = RHICreateTexture(Desc);
 			}
 			
 			DecodedTextureUAV = nullptr;
diff --git a/Engine/Plugins/Runtime/AR/Google/GoogleARCore/Source/GoogleARCoreBase/Private/GoogleARCoreTexture.cpp b/Engine/Plugins/Runtime/AR/Google/GoogleARCore/Source/GoogleARCoreBase/Private/GoogleARCoreTexture.cpp
--- a/Engine/Plugins/Runtime/AR/Google/GoogleARCore/Source/GoogleARCoreBase/Private/GoogleARCoreTexture.cpp
+++ b/Engine/Plugins/Runtime/AR/Google/GoogleARCore/Source/GoogleARCoreBase/Private/GoogleARCoreTexture.cpp
@@ -20,9 +20,12 @@ public:
 	{
 		FSamplerStateInitializerRHI SamplerStateInitializer(SF_Bilinear, AM_Clamp, AM_Clamp, AM_Clamp);
 		SamplerStateRHI = RHICreateSamplerState(SamplerStateInitializer);
-		
-		FRHIResourceCreateInfo CreateInfo(TEXT("FARCoreCameraTextureResource"));
-		TextureRHI = RHICreateTextureExternal2D(1, 1, PF_R8G8B8A8, 1, 1, TexCreate_SRGB, CreateInfo);
+
+		const FRHITextureCreateDesc Desc =
+			FRHITextureCreateDesc::Create2D(TEXT("FARCoreCameraTextureResource"), 1, 1, PF_R8G8B8A8)
+			.SetFlags(ETextureCreateFlags::External | ETextureCreateFlags::SRGB);
+
+		TextureRHI = RHICreateTexture(Desc);
 
 		ExecuteOnRHIThread([this]()
 		{
@@ -98,11 +101,14 @@ public:
 		SamplerStateRHI = RHICreateSamplerState(SamplerStateInitializer);
 		
 		FString Name = Owner->GetName();
-		FRHIResourceCreateInfo CreateInfo(*Name);
 		// Our source data is 16bits integer so PF_R32_FLOAT is a bit wasteful
 		// But PF_R16_UINT doesn't work in the material and it's not easy to convert the data to PF_R16F on the CPU
-		TextureRHI = RHICreateTexture2D(SizeX, SizeY, PF_R32_FLOAT, 1, 1, TexCreate_Dynamic | TexCreate_ShaderResource, CreateInfo);
-		
+		const FRHITextureCreateDesc Desc =
+			FRHITextureCreateDesc::Create2D(*Name, SizeX, SizeY, PF_R32_FLOAT)
+			.SetFlags(ETextureCreateFlags::Dynamic | ETextureCreateFlags::ShaderResource);
+
+		TextureRHI = RHICreateTexture(Desc);
+
 		TextureRHI->SetName(Owner->GetFName());
 		RHIBindDebugLabelName(TextureRHI, *Name);
 		RHIUpdateTextureReference(Owner->TextureReference.TextureReferenceRHI, TextureRHI);
diff --git a/Engine/Plugins/Runtime/AppleMoviePlayer/Source/AppleMoviePlayer/Private/AppleMovieStreamer.cpp b/Engine/Plugins/Runtime/AppleMoviePlayer/Source/AppleMoviePlayer/Private/AppleMovieStreamer.cpp
--- a/Engine/Plugins/Runtime/AppleMoviePlayer/Source/AppleMoviePlayer/Private/AppleMovieStreamer.cpp
+++ b/Engine/Plugins/Runtime/AppleMoviePlayer/Source/AppleMoviePlayer/Private/AppleMovieStreamer.cpp
@@ -591,10 +591,14 @@ bool FAVPlayerMovieStreamer::CheckForNextFrameAndCopy()
 		CVReturn Result = CVMetalTextureCacheCreateTextureFromImage(kCFAllocatorDefault, MetalTextureCache, pPixelBuffer, nullptr, MTLPixelFormatBGRA8Unorm, Width, Height, 0, &TextureCacheRef);
 		check(Result == kCVReturnSuccess);
 		check(TextureCacheRef);
-	
-		FRHIResourceCreateInfo CreateInfo(TEXT("FAVPlayerMovieStreamer"), new FAppleMovieStreamerTexture2DResourceWrapper(TextureCacheRef));
-		FTexture2DRHIRef MediaWrappedTexture = RHICreateTexture2D(Width, Height, PF_B8G8R8A8, 1, 1, TexCreate_Dynamic | TexCreate_NoTiling, CreateInfo);
-		
+
+		const FRHITextureCreateDesc Desc =
+			FRHITextureCreateDesc::Create2D(TEXT("FAVPlayerMovieStreamer"), Width, Height, PF_B8G8R8A8)
+			.SetFlags(ETextureCreateFlags::Dynamic | ETextureCreateFlags::NoTiling)
+			.SetBulkData(new FAppleMovieStreamerTexture2DResourceWrapper(TextureCacheRef));
+
+		FTextureRHIRef MediaWrappedTexture = RHICreateTexture(Desc);
+
 		if(SlateTexture.IsValid())
 		{
 			if(SlateTexture->IsValid())
diff --git a/Engine/Plugins/Runtime/HairStrands/Source/HairStrandsCore/Private/HairCardsBuilder.cpp b/Engine/Plugins/Runtime/HairStrands/Source/HairStrandsCore/Private/HairCardsBuilder.cpp
--- a/Engine/Plugins/Runtime/HairStrands/Source/HairStrandsCore/Private/HairCardsBuilder.cpp
+++ b/Engine/Plugins/Runtime/HairStrands/Source/HairStrandsCore/Private/HairCardsBuilder.cpp
@@ -4014,19 +4014,18 @@ static void AddCardsTextureReadbackPass(
 		check(OutTexture->Source.GetSizeX() == Resolution.X);
 		check(OutTexture->Source.GetSizeY() == Resolution.Y);
 
-		FRHIResourceCreateInfo CreateInfo(TEXT("CardsTextureReadbackPass_StagingTexture"));
-		FTexture2DRHIRef StagingTexture = RHICreateTexture2D(
-			InputTexture->Desc.Extent.X,
-			InputTexture->Desc.Extent.Y,
-			InputTexture->Desc.Format,
-			InputTexture->Desc.NumMips,
-			1, TexCreate_CPUReadback, CreateInfo);
+		const FRHITextureCreateDesc Desc =
+			FRHITextureCreateDesc::Create2D(TEXT("CardsTextureReadbackPass_StagingTexture"), InputTexture->Desc.Extent, InputTexture->Desc.Format)
+			.SetNumMips(InputTexture->Desc.NumMips)
+			.SetFlags(ETextureCreateFlags::CPUReadback);
+
+		FTextureRHIRef StagingTexture = RHICreateTexture(Desc);
 
 		FRHICopyTextureInfo CopyInfo;
 		CopyInfo.NumMips = InputTexture->Desc.NumMips;
 		RHICmdList.CopyTexture(
 			InputTexture->GetRHI(),
-			StagingTexture->GetTexture2D(),
+			StagingTexture,
 			CopyInfo);
 
 		// Flush, to ensure that all texture generation is done
diff --git a/Engine/Plugins/Runtime/TextureShare/Source/TextureShareD3D11/Private/TextureShareD3D11.cpp b/Engine/Plugins/Runtime/TextureShare/Source/TextureShareD3D11/Private/TextureShareD3D11.cpp
--- a/Engine/Plugins/Runtime/TextureShare/Source/TextureShareD3D11/Private/TextureShareD3D11.cpp
+++ b/Engine/Plugins/Runtime/TextureShare/Source/TextureShareD3D11/Private/TextureShareD3D11.cpp
@@ -62,8 +62,12 @@ bool FTextureShareD3D11::CreateRHITexture(ID3D11Texture2D* OpenedSharedResource,
 bool FTextureShareD3D11::CreateSharedTexture(FIntPoint& Size, EPixelFormat Format, FTexture2DRHIRef& OutRHITexture, void*& OutHandle)
 {
 	// Create RHI resource
-	FRHIResourceCreateInfo CreateInfo(TEXT("CreateSharedTexture"));
-	OutRHITexture = RHICreateTexture2D(Size.X, Size.Y, Format, 1, 1, TexCreate_Shared | TexCreate_ResolveTargetable, CreateInfo);
+	const FRHITextureCreateDesc Desc =
+		FRHITextureCreateDesc::Create2D(TEXT("CreateSharedTexture"), Size, Format)
+		.SetFlags(ETextureCreateFlags::Shared | ETextureCreateFlags::ResolveTargetable);
+
+	OutRHITexture = RHICreateTexture(Desc);
+
 	if (OutRHITexture.IsValid() && OutRHITexture->IsValid())
 	{
 		auto UE4D3DDevice = static_cast<ID3D11Device*>(GDynamicRHI->RHIGetNativeDevice());
diff --git a/Engine/Plugins/Runtime/TextureShare/Source/TextureShareD3D12/Private/TextureShareD3D12.cpp b/Engine/Plugins/Runtime/TextureShare/Source/TextureShareD3D12/Private/TextureShareD3D12.cpp
--- a/Engine/Plugins/Runtime/TextureShare/Source/TextureShareD3D12/Private/TextureShareD3D12.cpp
+++ b/Engine/Plugins/Runtime/TextureShare/Source/TextureShareD3D12/Private/TextureShareD3D12.cpp
@@ -37,11 +37,15 @@ bool FTextureShareD3D12::CreateRHITexture(ID3D12Resource* OpenedSharedResource,
 
 bool FTextureShareD3D12::CreateSharedTexture(FIntPoint& Size, EPixelFormat Format, FTexture2DRHIRef& OutRHITexture, void*& OutHandle, FGuid& OutSharedHandleGuid)
 {
-	FString UniqueHandleId = FString::Printf(TEXT("Global\\%s"), *OutSharedHandleGuid.ToString(EGuidFormats::DigitsWithHyphensInBraces));
+	const FString UniqueHandleId = FString::Printf(TEXT("Global\\%s"), *OutSharedHandleGuid.ToString(EGuidFormats::DigitsWithHyphensInBraces));
 
 	// Create RHI resource
-	FRHIResourceCreateInfo CreateInfo(*UniqueHandleId);
-	OutRHITexture = RHICreateTexture2D(Size.X, Size.Y, Format, 1, 1, TexCreate_Shared | TexCreate_ResolveTargetable, CreateInfo);
+	const FRHITextureCreateDesc Desc =
+		FRHITextureCreateDesc::Create2D(*UniqueHandleId, Size, Format)
+		.SetFlags(ETextureCreateFlags::Shared | ETextureCreateFlags::ResolveTargetable);
+
+	OutRHITexture = RHICreateTexture(Desc);
+
 	if (OutRHITexture.IsValid() && OutRHITexture->IsValid())
 	{
 		ID3D12Device* D3DDevice = GetID3D12DynamicRHI()->RHIGetDevice(0);
diff --git a/Engine/Plugins/Runtime/WindowsMixedReality/Source/HoloLensAR/Private/HoloLensCameraImageTexture.cpp b/Engine/Plugins/Runtime/WindowsMixedReality/Source/HoloLensAR/Private/HoloLensCameraImageTexture.cpp
--- a/Engine/Plugins/Runtime/WindowsMixedReality/Source/HoloLensAR/Private/HoloLensCameraImageTexture.cpp
+++ b/Engine/Plugins/Runtime/WindowsMixedReality/Source/HoloLensAR/Private/HoloLensCameraImageTexture.cpp
@@ -249,9 +249,13 @@ public:
 		// Default to an empty 1x1 texture if we don't have a camera image or failed to convert
 		if (!bDidConvert)
 		{
-			FRHIResourceCreateInfo CreateInfo(TEXT("DecodedTextureRef"));
 			Size.X = Size.Y = 1;
-			DecodedTextureRef = RHICreateTexture2D(Size.X, Size.Y, PF_B8G8R8A8, 1, 1, TexCreate_ShaderResource, CreateInfo);
+
+			const FRHITextureCreateDesc Desc =
+				FRHITextureCreateDesc::Create2D(TEXT("DecodedTextureRef"), Size, PF_B8G8R8A8)
+				.SetFlags(ETextureCreateFlags::ShaderResource);
+
+			DecodedTextureRef = RHICreateTexture(Desc);
 		}
 
 		TextureRHI = DecodedTextureRef;
diff --git a/Engine/Plugins/Runtime/WindowsMixedReality/Source/WindowsMixedRealityHMD/Private/WindowsMixedRealityHMD.cpp b/Engine/Plugins/Runtime/WindowsMixedReality/Source/WindowsMixedRealityHMD/Private/WindowsMixedRealityHMD.cpp
--- a/Engine/Plugins/Runtime/WindowsMixedReality/Source/WindowsMixedRealityHMD/Private/WindowsMixedRealityHMD.cpp
+++ b/Engine/Plugins/Runtime/WindowsMixedReality/Source/WindowsMixedRealityHMD/Private/WindowsMixedRealityHMD.cpp
@@ -1825,9 +1825,11 @@ namespace WindowsMixedReality
 		// Create a new texture for the remapped depth.
 		if (remappedDepthTexture == nullptr || recreateTextures)
 		{
-			FRHIResourceCreateInfo CreateInfo(TEXT("RemappedDepthTexture"));
-			remappedDepthTexture = RHICreateTexture2D(depthFRHITexture->GetSizeX(), depthFRHITexture->GetSizeY(),
-				PF_R32_FLOAT, 1, 1, TexCreate_RenderTargetable, CreateInfo);
+			const FRHITextureCreateDesc Desc =
+				FRHITextureCreateDesc::Create2D(TEXT("RemappedDepthTexture"), depthFRHITexture->GetSizeX(), depthFRHITexture->GetSizeY(), PF_R32_FLOAT)
+				.SetFlags(ETextureCreateFlags::RenderTargetable);
+
+			remappedDepthTexture = RHICreateTexture(Desc);
 		}
 
 		FGraphicsPipelineStateInitializer GraphicsPSOInit;
diff --git a/Engine/Plugins/Runtime/nDisplay/Source/DisplayCluster/Private/Render/Viewport/RenderTarget/DisplayClusterRenderTargetResource.cpp b/Engine/Plugins/Runtime/nDisplay/Source/DisplayCluster/Private/Render/Viewport/RenderTarget/DisplayClusterRenderTargetResource.cpp
--- a/Engine/Plugins/Runtime/nDisplay/Source/DisplayCluster/Private/Render/Viewport/RenderTarget/DisplayClusterRenderTargetResource.cpp
+++ b/Engine/Plugins/Runtime/nDisplay/Source/DisplayCluster/Private/Render/Viewport/RenderTarget/DisplayClusterRenderTargetResource.cpp
@@ -52,34 +52,29 @@ void FDisplayClusterViewportResource::ImplInitDynamicRHI_RenderTargetResource2D(
 
 void FDisplayClusterViewportResource::ImplInitDynamicRHI_TextureResource2D(FTexture2DRHIRef& OutTextureRHI)
 {
-	ETextureCreateFlags CreateFlags = TexCreate_Dynamic | TexCreate_ShaderResource;
-
-	// -- we will be manually copying this cross GPU, tell render graph not to
-	CreateFlags |= TexCreate_MultiGPUGraphIgnore;
+	FRHITextureCreateDesc Desc =
+		FRHITextureCreateDesc::Create2D(TEXT("DisplayClusterViewportTextureResource"), GetSizeX(), GetSizeY(), ViewportResourceSettings.Format)
+		.SetClearValue(FClearValueBinding::Black)
+		.SetFlags(ETextureCreateFlags::Dynamic | ETextureCreateFlags::ShaderResource)
+		// -- we will be manually copying this cross GPU, tell render graph not to
+		.AddFlags(ETextureCreateFlags::MultiGPUGraphIgnore)
+		.SetInitialState(ERHIAccess::SRVMask);
 
 	// reflect srgb from settings
 	if (bSRGB)
 	{
-		CreateFlags |= TexCreate_SRGB;
+		Desc.AddFlags(ETextureCreateFlags::SRGB);
+	}
+	if (ViewportResourceSettings.bIsRenderTargetable)
+	{
+		Desc.AddFlags(ETextureCreateFlags::RenderTargetable);
+	}
+	else
+	{
+		Desc.AddFlags(ETextureCreateFlags::ResolveTargetable);
 	}
 
-	uint32 NumMips = 1;
-	uint32 NumSamples = 1;
-
-	CreateFlags |= ViewportResourceSettings.bIsRenderTargetable ? TexCreate_RenderTargetable : TexCreate_ResolveTargetable;
-
-	FRHIResourceCreateInfo CreateInfo(TEXT("DisplayClusterViewportTextureResource"), FClearValueBinding::Black);
-
-	OutTextureRHI = RHICreateTexture2D(
-		GetSizeX(),
-		GetSizeY(),
-		ViewportResourceSettings.Format,
-		NumMips,
-		NumSamples,
-		CreateFlags,
-		ERHIAccess::SRVMask,
-		CreateInfo
-	);
+	OutTextureRHI = RHICreateTexture(Desc);
 }
 
 //------------------------------------------------------------------------------------------------
diff --git a/Engine/Plugins/Runtime/nDisplay/Source/DisplayClusterShaders/Private/Render/DisplayClusterRenderTexture.cpp b/Engine/Plugins/Runtime/nDisplay/Source/DisplayClusterShaders/Private/Render/DisplayClusterRenderTexture.cpp
--- a/Engine/Plugins/Runtime/nDisplay/Source/DisplayClusterShaders/Private/Render/DisplayClusterRenderTexture.cpp
+++ b/Engine/Plugins/Runtime/nDisplay/Source/DisplayClusterShaders/Private/Render/DisplayClusterRenderTexture.cpp
@@ -85,9 +85,12 @@ void FDisplayClusterRenderTextureResource::InitRHI()
 	FDisplayClusterResourceBulkData BulkDataInterface(TextureData, DataSize);
 
 	// @todo: Changed for CIS but needs to be fixed.
-	FRHIResourceCreateInfo CreateInfo(TEXT("DisplayClusterRenderTexture"), &BulkDataInterface);
+	const FRHITextureCreateDesc Desc =
+		FRHITextureCreateDesc::Create2D(TEXT("DisplayClusterRenderTexture"), Width, Height, PixelFormat)
+		.SetFlags(ETextureCreateFlags::ShaderResource)
+		.SetBulkData(&BulkDataInterface);
 
-	TextureRHI = RHICreateTexture2D(Width, Height, PixelFormat, 1, 1, TexCreate_ShaderResource, CreateInfo);
+	TextureRHI = RHICreateTexture(Desc);
 
 	// CPU access not required, release from memory
 	if (bHasCPUAccess == false)
diff --git a/Engine/Plugins/Tests/RHITests/Source/RHITests/Private/RHITextureTests.cpp b/Engine/Plugins/Tests/RHITests/Source/RHITests/Private/RHITextureTests.cpp
--- a/Engine/Plugins/Tests/RHITests/Source/RHITests/Private/RHITextureTests.cpp
+++ b/Engine/Plugins/Tests/RHITests/Source/RHITests/Private/RHITextureTests.cpp
@@ -15,11 +15,15 @@ bool FRHITextureTests::VerifyTextureContents(const TCHAR* TestName, FRHICommandL
 			uint32 MipHeight = FMath::Max(Size.Y >> MipIndex, 1);
 			uint32 MipDepth = FMath::Max(Size.Z >> MipIndex, 1);
 
+			const FRHITextureCreateDesc Desc =
+				FRHITextureCreateDesc::Create2D(TEXT("FRHITextureTests_StagingTexture"), MipWidth, MipHeight, Texture->GetFormat())
+				.SetFlags(ETextureCreateFlags::CPUReadback)
+				.SetInitialState(ERHIAccess::CopyDest);
+
 			for (uint32 Z = 0; Z < MipDepth; ++Z)
 			{
 				{
-					FRHIResourceCreateInfo CreateInfo(TEXT("FRHITextureTests_StagingTexture"));
-					FTexture2DRHIRef StagingTexture = RHICreateTexture2D(MipWidth, MipHeight, Texture->GetFormat(), 1, 1, TexCreate_CPUReadback, ERHIAccess::CopyDest, CreateInfo);
+					FTextureRHIRef StagingTexture = RHICreateTexture(Desc);
 
 					FRHICopyTextureInfo CopyInfo = {};
 					CopyInfo.Size = FIntVector(MipWidth, MipHeight, 1); // @todo - required for D3D11RHI to prevent crash in FD3D11DynamicRHI::RHICopyTexture
diff --git a/Engine/Source/Developer/MaterialBaking/Private/MaterialBakingModule.cpp b/Engine/Source/Developer/MaterialBaking/Private/MaterialBakingModule.cpp
--- a/Engine/Source/Developer/MaterialBaking/Private/MaterialBakingModule.cpp
+++ b/Engine/Source/Developer/MaterialBaking/Private/MaterialBakingModule.cpp
@@ -190,14 +190,17 @@ namespace FMaterialBakingModuleImpl
 			}
 
 			TRACE_CPUPROFILER_EVENT_SCOPE(RHICreateTexture2D)
-			FRHIResourceCreateInfo CreateInfo(TEXT("FStagingBufferPool_StagingBuffer"));
-			ETextureCreateFlags TextureCreateFlags = TexCreate_CPUReadback;
+
+			FRHITextureCreateDesc Desc =
+				FRHITextureCreateDesc::Create2D(TEXT("FStagingBufferPool_StagingBuffer"), Width, Height, Format)
+				.SetFlags(ETextureCreateFlags::CPUReadback);
+
 			if (bIsSRGB)
 			{
-				TextureCreateFlags |= TexCreate_SRGB;
+				Desc.AddFlags(ETextureCreateFlags::SRGB);
 			}
-			
-			return RHICreateTexture2D(Width, Height, Format, 1, 1, TextureCreateFlags, CreateInfo);
+
+			return RHICreateTexture(Desc);
 		}
 
 		void ReleaseStagingBufferForUnmap_AnyThread(FTexture2DRHIRef& Texture2DRHIRef)
diff --git a/Engine/Source/Editor/VirtualTexturingEditor/Private/RuntimeVirtualTextureBuildStreamingMips.cpp b/Engine/Source/Editor/VirtualTexturingEditor/Private/RuntimeVirtualTextureBuildStreamingMips.cpp
--- a/Engine/Source/Editor/VirtualTexturingEditor/Private/RuntimeVirtualTextureBuildStreamingMips.cpp
+++ b/Engine/Source/Editor/VirtualTexturingEditor/Private/RuntimeVirtualTextureBuildStreamingMips.cpp
@@ -46,9 +46,14 @@ namespace
 
 			for (int32 Layer = 0; Layer < NumLayers; ++Layer)
 			{
-				FRHIResourceCreateInfo CreateInfo(TEXT("FTileRenderResources"));
-				RenderTargets[Layer] = RHICreateTexture2D(TileSize, TileSize, LayerFormats[Layer], 1, 1, TexCreate_RenderTargetable, CreateInfo);
-				StagingTextures[Layer] = RHICreateTexture2D(TileSize, TileSize, LayerFormats[Layer], 1, 1, TexCreate_CPUReadback, CreateInfo);
+				FRHITextureCreateDesc Desc =
+					FRHITextureCreateDesc::Create2D(TEXT("FTileRenderResources"), TileSize, TileSize, LayerFormats[Layer]);
+
+				Desc.SetFlags(ETextureCreateFlags::RenderTargetable);
+				RenderTargets[Layer] = RHICreateTexture(Desc);
+
+				Desc.SetFlags(ETextureCreateFlags::CPUReadback);
+				StagingTextures[Layer] = RHICreateTexture(Desc);
 			}
 
 			FRHICommandListImmediate& RHICmdList = FRHICommandListExecutor::GetImmediateCommandList();
diff --git a/Engine/Source/Runtime/Apple/AGXRHI/Private/AGXStateCache.cpp b/Engine/Source/Runtime/Apple/AGXRHI/Private/AGXStateCache.cpp
--- a/Engine/Source/Runtime/Apple/AGXRHI/Private/AGXStateCache.cpp
+++ b/Engine/Source/Runtime/Apple/AGXRHI/Private/AGXStateCache.cpp
@@ -2239,8 +2239,11 @@ FTexture2DRHIRef FAGXStateCache::CreateFallbackDepthStencilSurface(uint32 Width,
 	if (!IsValidRef(FallbackDepthStencilSurface) || FallbackDepthStencilSurface->GetSizeX() != Width || FallbackDepthStencilSurface->GetSizeY() != Height)
 #endif
 	{
-		FRHIResourceCreateInfo TexInfo(TEXT("FallbackDepthStencilSurface"));
-		FallbackDepthStencilSurface = RHICreateTexture2D(Width, Height, PF_DepthStencil, 1, 1, TexCreate_DepthStencilTargetable, TexInfo);
+		const FRHITextureCreateDesc Desc =
+			FRHITextureCreateDesc::Create2D(TEXT("FallbackDepthStencilSurface"), Width, Height, PF_DepthStencil)
+			.SetFlags(ETextureCreateFlags::DepthStencilTargetable);
+
+		FallbackDepthStencilSurface = RHICreateTexture(Desc);
 	}
 	check(IsValidRef(FallbackDepthStencilSurface));
 	return FallbackDepthStencilSurface;
diff --git a/Engine/Source/Runtime/Apple/MetalRHI/Private/MetalStateCache.cpp b/Engine/Source/Runtime/Apple/MetalRHI/Private/MetalStateCache.cpp
--- a/Engine/Source/Runtime/Apple/MetalRHI/Private/MetalStateCache.cpp
+++ b/Engine/Source/Runtime/Apple/MetalRHI/Private/MetalStateCache.cpp
@@ -2251,8 +2251,11 @@ FTexture2DRHIRef FMetalStateCache::CreateFallbackDepthStencilSurface(uint32 Widt
 	if (!IsValidRef(FallbackDepthStencilSurface) || FallbackDepthStencilSurface->GetSizeX() != Width || FallbackDepthStencilSurface->GetSizeY() != Height)
 #endif
 	{
-		FRHIResourceCreateInfo TexInfo(TEXT("FallbackDepthStencilSurface"));
-		FallbackDepthStencilSurface = RHICreateTexture2D(Width, Height, PF_DepthStencil, 1, 1, TexCreate_DepthStencilTargetable, TexInfo);
+		const FRHITextureCreateDesc Desc =
+			FRHITextureCreateDesc::Create2D(TEXT("FallbackDepthStencilSurface"), Width, Height, PF_DepthStencil)
+			.SetFlags(ETextureCreateFlags::DepthStencilTargetable);
+
+		FallbackDepthStencilSurface = RHICreateTexture(Desc);
 	}
 	check(IsValidRef(FallbackDepthStencilSurface));
 	return FallbackDepthStencilSurface;
diff --git a/Engine/Source/Runtime/Engine/Private/Components/LightComponent.cpp b/Engine/Source/Runtime/Engine/Private/Components/LightComponent.cpp
--- a/Engine/Source/Runtime/Engine/Private/Components/LightComponent.cpp
+++ b/Engine/Source/Runtime/Engine/Private/Components/LightComponent.cpp
@@ -35,12 +35,13 @@ void FStaticShadowDepthMap::InitRHI()
 {
 	if (FApp::CanEverRender() && Data && Data->ShadowMapSizeX > 0 && Data->ShadowMapSizeY > 0 && GMaxRHIFeatureLevel >= ERHIFeatureLevel::SM5)
 	{
-		FRHIResourceCreateInfo CreateInfo(TEXT("FStaticShadowDepthMap"));
-		FTexture2DRHIRef Texture2DRHI = RHICreateTexture2D(Data->ShadowMapSizeX, Data->ShadowMapSizeY, PF_R16F, 1, 1, TexCreate_None, CreateInfo);
-		TextureRHI = Texture2DRHI;
+		const FRHITextureCreateDesc Desc =
+			FRHITextureCreateDesc::Create2D(TEXT("FStaticShadowDepthMap"), Data->ShadowMapSizeX, Data->ShadowMapSizeY, PF_R16F);
+
+		TextureRHI = RHICreateTexture(Desc);
 
 		uint32 DestStride = 0;
-		uint8* TextureData = (uint8*)RHILockTexture2D(Texture2DRHI, 0, RLM_WriteOnly, DestStride, false);
+		uint8* TextureData = (uint8*)RHILockTexture2D(TextureRHI, 0, RLM_WriteOnly, DestStride, false);
 		uint32 RowSize = Data->ShadowMapSizeX * GPixelFormats[PF_R16F].BlockBytes;
 
 		for (int32 Y = 0; Y < Data->ShadowMapSizeY; Y++)
@@ -48,7 +49,7 @@ void FStaticShadowDepthMap::InitRHI()
 			FMemory::Memcpy(TextureData + DestStride * Y, ((uint8*)Data->DepthSamples.GetData()) + RowSize * Y, RowSize);
 		}
 
-		RHIUnlockTexture2D(Texture2DRHI, 0, false);
+		RHIUnlockTexture2D(TextureRHI, 0, false);
 	}
 }
 
diff --git a/Engine/Source/Runtime/Engine/Private/PrimitiveUniformShaderParameters.cpp b/Engine/Source/Runtime/Engine/Private/PrimitiveUniformShaderParameters.cpp
--- a/Engine/Source/Runtime/Engine/Private/PrimitiveUniformShaderParameters.cpp
+++ b/Engine/Source/Runtime/Engine/Private/PrimitiveUniformShaderParameters.cpp
@@ -35,8 +35,11 @@ void FSinglePrimitiveStructured::InitRHI()
 		}
 
 		{
-			CreateInfo.DebugName = TEXT("PrimitiveSceneDataTexture");
-			PrimitiveSceneDataTextureRHI = RHICreateTexture2D(FPrimitiveSceneShaderData::DataStrideInFloat4s, 1, PF_A32B32G32R32F, 1, 1, TexCreate_ShaderResource | TexCreate_UAV, CreateInfo);
+			const FRHITextureCreateDesc Desc =
+				FRHITextureCreateDesc::Create2D(TEXT("PrimitiveSceneDataTexture"), FPrimitiveSceneShaderData::DataStrideInFloat4s, 1, PF_A32B32G32R32F)
+				.SetFlags(ETextureCreateFlags::ShaderResource | ETextureCreateFlags::UAV);
+
+			PrimitiveSceneDataTextureRHI = RHICreateTexture(Desc);
 			PrimitiveSceneDataTextureSRV = RHICreateShaderResourceView(PrimitiveSceneDataTextureRHI, 0);
 		}
 
diff --git a/Engine/Source/Runtime/Engine/Private/Rendering/Texture2DResource.cpp b/Engine/Source/Runtime/Engine/Private/Rendering/Texture2DResource.cpp
--- a/Engine/Source/Runtime/Engine/Private/Rendering/Texture2DResource.cpp
+++ b/Engine/Source/Runtime/Engine/Private/Rendering/Texture2DResource.cpp
@@ -100,14 +100,18 @@ void FTexture2DResource::InitRHI()
 
 void FTexture2DResource::CreateTexture()
 {
-	FTexture2DRHIRef Texture2DRHI;
 	const int32 RequestedMipIdx = State.RequestedFirstLODIdx();
 	const FTexture2DMipMap* RequestedMip = GetPlatformMip(RequestedMipIdx);
 
 	// create texture with ResourceMem data when available
-	FRHIResourceCreateInfo CreateInfo(TEXT("FTexture2DResource"), ResourceMem);
-	CreateInfo.ExtData = PlatformData->GetExtData();
-	Texture2DRHI = RHICreateTexture2D( RequestedMip->SizeX, RequestedMip->SizeY, PixelFormat, State.NumRequestedLODs, 1, CreationFlags, CreateInfo);
+	const FRHITextureCreateDesc Desc =
+		FRHITextureCreateDesc::Create2D(TEXT("FTexture2DResource"), RequestedMip->SizeX, RequestedMip->SizeY, PixelFormat)
+		.SetNumMips(State.NumRequestedLODs)
+		.SetFlags(CreationFlags)
+		.SetExtData(PlatformData->GetExtData())
+		.SetBulkData(ResourceMem);
+
+	TextureRHI = RHICreateTexture(Desc);
 
 	// if( ResourceMem && !State.bReadyForStreaming) // To investigate!
 	if (ResourceMem)
@@ -129,26 +133,31 @@ void FTexture2DResource::CreateTexture()
 			if (MipData[ResourceMipIdx])
 			{
 				uint32 DestPitch = -1;
-				void* TheMipData = RHILockTexture2D( Texture2DRHI, RHIMipIdx, RLM_WriteOnly, DestPitch, false );
+				void* TheMipData = RHILockTexture2D(TextureRHI, RHIMipIdx, RLM_WriteOnly, DestPitch, false );
 				GetData( ResourceMipIdx, TheMipData, DestPitch );
-				RHIUnlockTexture2D( Texture2DRHI, RHIMipIdx, false );
+				RHIUnlockTexture2D(TextureRHI, RHIMipIdx, false );
 			}
 		}
 	}
-	TextureRHI = Texture2DRHI;
 }
 
 void FTexture2DResource::CreatePartiallyResidentTexture()
 {
-	FTexture2DRHIRef Texture2DRHI;
 	const int32 CurrentFirstMip = State.RequestedFirstLODIdx();
 
 	check(bUsePartiallyResidentMips);
-	FRHIResourceCreateInfo CreateInfo(TEXT("FTexture2DResource-PRT"), ResourceMem);
-	CreateInfo.ExtData = PlatformData->GetExtData();
-	Texture2DRHI = RHICreateTexture2D( SizeX, SizeY, PixelFormat, State.MaxNumLODs, 1, CreationFlags | TexCreate_Virtual, CreateInfo);
-	RHIVirtualTextureSetFirstMipInMemory(Texture2DRHI, CurrentFirstMip);
-	RHIVirtualTextureSetFirstMipVisible(Texture2DRHI, CurrentFirstMip);
+
+	const FRHITextureCreateDesc Desc =
+		FRHITextureCreateDesc::Create2D(TEXT("FTexture2DResource-PRT"), SizeX, SizeY, PixelFormat)
+		.SetNumMips(State.MaxNumLODs)
+		.SetFlags(CreationFlags | ETextureCreateFlags::Virtual)
+		.SetExtData(PlatformData->GetExtData())
+		.SetBulkData(ResourceMem);
+
+	TextureRHI = RHICreateTexture(Desc);
+
+	RHIVirtualTextureSetFirstMipInMemory(TextureRHI, CurrentFirstMip);
+	RHIVirtualTextureSetFirstMipVisible(TextureRHI, CurrentFirstMip);
 
 	check( ResourceMem == nullptr );
 
@@ -158,13 +167,11 @@ void FTexture2DResource::CreatePartiallyResidentTexture()
 		if ( MipData[MipIndex] != NULL )
 		{
 			uint32 DestPitch = -1;
-			void* TheMipData = RHILockTexture2D( Texture2DRHI, MipIndex, RLM_WriteOnly, DestPitch, false );
+			void* TheMipData = RHILockTexture2D(TextureRHI, MipIndex, RLM_WriteOnly, DestPitch, false );
 			GetData( MipIndex, TheMipData, DestPitch );
-			RHIUnlockTexture2D( Texture2DRHI, MipIndex, false );
+			RHIUnlockTexture2D(TextureRHI, MipIndex, false );
 		}
 	}
-
-	TextureRHI = Texture2DRHI;
 }
 
 uint64 FTexture2DResource::GetPlatformMipsSize(uint32 NumMips) const
diff --git a/Engine/Source/Runtime/Engine/Private/Slate/SlateTextures.cpp b/Engine/Source/Runtime/Engine/Private/Slate/SlateTextures.cpp
--- a/Engine/Source/Runtime/Engine/Private/Slate/SlateTextures.cpp
+++ b/Engine/Source/Runtime/Engine/Private/Slate/SlateTextures.cpp
@@ -49,8 +49,12 @@ void FSlateTexture2DRHIRef::InitDynamicRHI()
 		if( TextureData.IsValid() || bCreateEmptyTexture )
 		{
 			check( !IsValidRef( ShaderResource) );
-			FRHIResourceCreateInfo CreateInfo(TEXT("FSlateTexture2DRHIRef"));
-			ShaderResource = RHICreateTexture2D( Width, Height, PixelFormat, 1, 1, TexCreateFlags, CreateInfo );
+
+			const FRHITextureCreateDesc Desc =
+				FRHITextureCreateDesc::Create2D(TEXT("FSlateTexture2DRHIRef"), Width, Height, PixelFormat)
+				.SetFlags(TexCreateFlags);
+
+			ShaderResource = RHICreateTexture(Desc);
 			check( IsValidRef( ShaderResource ) );
 
 			INC_MEMORY_STAT_BY(STAT_SlateTextureGPUMemory, Width*Height*GPixelFormats[PixelFormat].BlockBytes);
diff --git a/Engine/Source/Runtime/Engine/Private/Streaming/Texture2DUpdate.cpp b/Engine/Source/Runtime/Engine/Private/Streaming/Texture2DUpdate.cpp
--- a/Engine/Source/Runtime/Engine/Private/Streaming/Texture2DUpdate.cpp
+++ b/Engine/Source/Runtime/Engine/Private/Streaming/Texture2DUpdate.cpp
@@ -94,15 +94,14 @@ void FTexture2DUpdate::DoConvertToVirtualWithNewMips(const FContext& Context)
 			ensure(!IntermediateTextureRHI);
 
 			// Create a copy of the texture that is a virtual texture.
-			FRHIResourceCreateInfo CreateInfo(TEXT("FTexture2DUpdate"), Context.Resource->ResourceMem);
-			IntermediateTextureRHI = RHICreateTexture2D(
-				MipMap0.SizeX, 
-				MipMap0.SizeY, 
-				Context.Resource->GetPixelFormat(), 
-				ResourceState.MaxNumLODs, 
-				1, 
-				Context.Resource->GetCreationFlags() | TexCreate_Virtual, 
-				CreateInfo);
+			const FRHITextureCreateDesc Desc =
+				FRHITextureCreateDesc::Create2D(TEXT("FTexture2DUpdate"), MipMap0.SizeX, MipMap0.SizeY, Context.Resource->GetPixelFormat())
+				.SetNumMips(ResourceState.MaxNumLODs)
+				.SetFlags(Context.Resource->GetCreationFlags() | ETextureCreateFlags::Virtual)
+				.SetBulkData(Context.Resource->ResourceMem);
+
+			IntermediateTextureRHI = RHICreateTexture(Desc);
+
 			RHIVirtualTextureSetFirstMipInMemory(IntermediateTextureRHI, CurrentFirstLODIdx);
 			RHIVirtualTextureSetFirstMipVisible(IntermediateTextureRHI, CurrentFirstLODIdx);
 			RHICopySharedMips(IntermediateTextureRHI, Context.Resource->GetTexture2DRHI());
@@ -128,15 +127,15 @@ bool FTexture2DUpdate::DoConvertToNonVirtual(const FContext& Context)
 			const FTexture2DMipMap& PendingFirstMipMap = *Context.MipsView[PendingFirstLODIdx];
 
 			ensure(!IntermediateTextureRHI);
-			FRHIResourceCreateInfo CreateInfo(TEXT("FTexture2DUpdate"), Context.Resource->ResourceMem);
-			IntermediateTextureRHI = RHICreateTexture2D(
-				PendingFirstMipMap.SizeX, 
-				PendingFirstMipMap.SizeY, 
-				Context.Resource->GetPixelFormat(), 
-				ResourceState.NumRequestedLODs,
-				1, 
-				Context.Resource->GetCreationFlags(), 
-				CreateInfo);
+
+			const FRHITextureCreateDesc Desc =
+				FRHITextureCreateDesc::Create2D(TEXT("FTexture2DUpdate"), PendingFirstMipMap.SizeX, PendingFirstMipMap.SizeY, Context.Resource->GetPixelFormat())
+				.SetNumMips(ResourceState.NumRequestedLODs)
+				.SetFlags(Context.Resource->GetCreationFlags())
+				.SetBulkData(Context.Resource->ResourceMem);
+
+			IntermediateTextureRHI = RHICreateTexture(Desc);
+
 			RHICopySharedMips(IntermediateTextureRHI, Context.Resource->GetTexture2DRHI());
 
 			return true;
diff --git a/Engine/Source/Runtime/Engine/Private/Texture2D.cpp b/Engine/Source/Runtime/Engine/Private/Texture2D.cpp
--- a/Engine/Source/Runtime/Engine/Private/Texture2D.cpp
+++ b/Engine/Source/Runtime/Engine/Private/Texture2D.cpp
@@ -1396,15 +1396,26 @@ void FVirtualTexture2DResource::InitializeEditorResources(IVirtualTexture* InVir
 			}
 		}
 
-		ETextureCreateFlags TexCreateFlags = (TextureOwner->SRGB ? TexCreate_SRGB : TexCreate_None) | (TextureOwner->bNotOfflineProcessed ? TexCreate_None : TexCreate_OfflineProcessed);
+		FString Name = TextureOwner->GetName();
+
+		FRHITextureCreateDesc Desc =
+			FRHITextureCreateDesc::Create2D(*Name, MipWidthInTiles * TileSizeInPixels, MipHeightInTiles * TileSizeInPixels, PixelFormat);
+
+		if (TextureOwner->SRGB)
+		{
+			Desc.AddFlags(ETextureCreateFlags::SRGB);
+		}
+		if (!TextureOwner->bNotOfflineProcessed)
+		{
+			Desc.AddFlags(ETextureCreateFlags::OfflineProcessed);
+		}
 		if (TextureOwner->bNoTiling)
 		{
-			TexCreateFlags |= TexCreate_NoTiling;
+			Desc.AddFlags(ETextureCreateFlags::NoTiling);
 		}
 
-		FString Name = TextureOwner->GetName();
-		FRHIResourceCreateInfo CreateInfo(*Name);
-		FTexture2DRHIRef Texture2DRHI = RHICreateTexture2D(MipWidthInTiles * TileSizeInPixels, MipHeightInTiles * TileSizeInPixels, PixelFormat, 1, 1, TexCreateFlags, CreateInfo);
+		FTexture2DRHIRef Texture2DRHI = RHICreateTexture(Desc);
+
 		FRHICommandListImmediate& RHICommandList = FRHICommandListExecutor::GetImmediateCommandList();
 
 		TArray<IVirtualTextureFinalizer*> Finalizers;
@@ -1445,7 +1456,13 @@ void FVirtualTexture2DResource::InitializeEditorResources(IVirtualTexture* InVir
 			check(MipWidth <= MipWidthInTiles * TileSizeInPixels);
 			check(MipHeight <= MipHeightInTiles * TileSizeInPixels);
 
-			FTexture2DRHIRef ResizedTexture2DRHI = RHICreateTexture2D(MipWidth, MipHeight, PixelFormat, 1, 1, TexCreateFlags, ERHIAccess::CopyDest, CreateInfo);
+			const FRHITextureCreateDesc ResizedDesc =
+				FRHITextureCreateDesc::Create2D(*Name, MipWidth, MipHeight, PixelFormat)
+				.SetFlags(Desc.Flags)
+				.SetInitialState(ERHIAccess::CopyDest);
+
+			FTexture2DRHIRef ResizedTexture2DRHI = RHICreateTexture(ResizedDesc);
+
 			FRHICopyTextureInfo CopyInfo;
 			CopyInfo.Size = FIntVector(MipWidth, MipHeight, 1);
 
diff --git a/Engine/Source/Runtime/Engine/Private/Texture2DDynamic.cpp b/Engine/Source/Runtime/Engine/Private/Texture2DDynamic.cpp
--- a/Engine/Source/Runtime/Engine/Private/Texture2DDynamic.cpp
+++ b/Engine/Source/Runtime/Engine/Private/Texture2DDynamic.cpp
@@ -46,23 +46,29 @@ void FTexture2DDynamicResource::InitRHI()
 	);
 	SamplerStateRHI = GetOrCreateSamplerState( SamplerStateInitializer );
 
-	ETextureCreateFlags  Flags = TexCreate_None;
-	if ( Owner->bIsResolveTarget )
+	FString Name = Owner->GetName();
+
+	FRHITextureCreateDesc Desc =
+		FRHITextureCreateDesc::Create2D(*Name, GetSizeX(), GetSizeY(), Owner->Format)
+		.SetNumMips(Owner->NumMips);
+
+	if (Owner->bIsResolveTarget)
 	{
-		Flags |= TexCreate_ResolveTargetable;
+		Desc.AddFlags(ETextureCreateFlags::ResolveTargetable);
 		bIgnoreGammaConversions = true;		// Note, we're ignoring Owner->SRGB (it should be false).
 	}
-	else if ( Owner->SRGB )
+	else if (Owner->SRGB)
 	{
-		Flags |= TexCreate_SRGB;
+		Desc.AddFlags(ETextureCreateFlags::SRGB);
 	}
-	if ( Owner->bNoTiling )
+
+	if (Owner->bNoTiling)
 	{
-		Flags |= TexCreate_NoTiling;
+		Desc.AddFlags(ETextureCreateFlags::NoTiling);
 	}
-	FString Name = Owner->GetName();
-	FRHIResourceCreateInfo CreateInfo(*Name);
-	Texture2DRHI = RHICreateTexture2D(GetSizeX(), GetSizeY(), Owner->Format, Owner->NumMips, 1, Flags, CreateInfo);
+
+	Texture2DRHI = RHICreateTexture(Desc);
+
 	TextureRHI = Texture2DRHI;
 	TextureRHI->SetName(Owner->GetFName());
 	RHIUpdateTextureReference(Owner->TextureReference.TextureReferenceRHI,TextureRHI);
diff --git a/Engine/Source/Runtime/Engine/Private/VT/VirtualTextureUploadCache.cpp b/Engine/Source/Runtime/Engine/Private/VT/VirtualTextureUploadCache.cpp
--- a/Engine/Source/Runtime/Engine/Private/VT/VirtualTextureUploadCache.cpp
+++ b/Engine/Source/Runtime/Engine/Private/VT/VirtualTextureUploadCache.cpp
@@ -303,8 +303,16 @@ void FVirtualTextureUploadCache::Finalize(FRDGBuilder& GraphBuilder)
 				DEC_MEMORY_STAT_BY(STAT_TotalGPUUploadSize, CalcTextureSize(StagingTexture.RHITexture->GetSizeX(), StagingTexture.RHITexture->GetSizeY(), PoolEntry.Format, 1u));
 			}
 
-			FRHIResourceCreateInfo CreateInfo(TEXT("FVirtualTextureUploadCache_StagingTexture"));
-			StagingTexture.RHITexture = RHICreateTexture2D(TileSize * WidthInTiles, TileSize * HeightInTiles, PoolEntry.Format, 1, 1, bIsCpuWritable ? TexCreate_CPUWritable : TexCreate_None, CreateInfo);
+			FRHITextureCreateDesc Desc =
+				FRHITextureCreateDesc::Create2D(TEXT("FVirtualTextureUploadCache_StagingTexture"), TileSize * WidthInTiles, TileSize * HeightInTiles, PoolEntry.Format);
+
+			if (bIsCpuWritable)
+			{
+				Desc.AddFlags(ETextureCreateFlags::CPUWritable);
+			}
+
+			StagingTexture.RHITexture = RHICreateTexture(Desc);
+
 			StagingTexture.WidthInTiles = WidthInTiles;
 			StagingTexture.BatchCapacity = WidthInTiles * HeightInTiles;
 			StagingTexture.bIsCPUWritable = bIsCpuWritable;
diff --git a/Engine/Source/Runtime/Engine/Public/PrecomputedVolumetricLightmap.h b/Engine/Source/Runtime/Engine/Public/PrecomputedVolumetricLightmap.h
--- a/Engine/Source/Runtime/Engine/Public/PrecomputedVolumetricLightmap.h
+++ b/Engine/Source/Runtime/Engine/Public/PrecomputedVolumetricLightmap.h
@@ -66,7 +66,7 @@ public:
 	// Stored redundantly for stats after Data has been discarded
 	int32 DataSize;
 	EPixelFormat Format;
-	FTexture3DRHIRef Texture;
+	FTextureRHIRef Texture;
 	FUnorderedAccessViewRHIRef UAV;
 
 	bool bNeedsCPUAccess;
diff --git a/Engine/Source/Runtime/Engine/Public/PrimitiveUniformShaderParameters.h b/Engine/Source/Runtime/Engine/Public/PrimitiveUniformShaderParameters.h
--- a/Engine/Source/Runtime/Engine/Public/PrimitiveUniformShaderParameters.h
+++ b/Engine/Source/Runtime/Engine/Public/PrimitiveUniformShaderParameters.h
@@ -567,7 +567,7 @@ public:
 	FBufferRHIRef InstancePayloadDataBufferRHI;
 	FShaderResourceViewRHIRef InstancePayloadDataBufferSRV;
 
-	FTexture2DRHIRef PrimitiveSceneDataTextureRHI;
+	FTextureRHIRef PrimitiveSceneDataTextureRHI;
 	FShaderResourceViewRHIRef PrimitiveSceneDataTextureSRV;
 
 	FBufferRHIRef LightmapSceneDataBufferRHI;
diff --git a/Engine/Source/Runtime/GameplayMediaEncoder/Private/GameplayMediaEncoder.cpp b/Engine/Source/Runtime/GameplayMediaEncoder/Private/GameplayMediaEncoder.cpp
--- a/Engine/Source/Runtime/GameplayMediaEncoder/Private/GameplayMediaEncoder.cpp
+++ b/Engine/Source/Runtime/GameplayMediaEncoder/Private/GameplayMediaEncoder.cpp
@@ -583,19 +583,25 @@ TSharedPtr<AVEncoder::FVideoEncoderInputFrame> FGameplayMediaEncoder::ObtainInpu
 	{
 #if PLATFORM_WINDOWS && PLATFORM_DESKTOP
 		const ERHIInterfaceType RHIType = RHIGetInterfaceType();
+
+		const FRHITextureCreateDesc Desc =
+			FRHITextureCreateDesc::Create2D(TEXT("VideoCapturerBackBuffer"), VideoConfig.Width, VideoConfig.Height, PF_B8G8R8A8)
+			.SetFlags(ETextureCreateFlags::Shared | ETextureCreateFlags::RenderTargetable | ETextureCreateFlags::UAV)
+			.SetInitialState(ERHIAccess::CopyDest);
+
 		if (RHIType == ERHIInterfaceType::D3D11)
 		{
 			FRHIResourceCreateInfo CreateInfo(TEXT("VideoCapturerBackBuffer"));
-			FTexture2DRHIRef Texture = RHICreateTexture2D(VideoConfig.Width, VideoConfig.Height, EPixelFormat::PF_B8G8R8A8, 1, 1,
-			                                                           TexCreate_Shared | TexCreate_RenderTargetable | TexCreate_UAV, ERHIAccess::CopyDest, CreateInfo);
+			FTextureRHIRef Texture = RHICreateTexture(Desc);
+
 			InputFrame->SetTexture((ID3D11Texture2D*)Texture->GetNativeResource(), [&, InputFrame](ID3D11Texture2D* NativeTexture) { BackBuffers.Remove(InputFrame); });
 			BackBuffers.Add(InputFrame, Texture);
 		}
 		else if (RHIType == ERHIInterfaceType::D3D12)
 		{
 			FRHIResourceCreateInfo CreateInfo(TEXT("VideoCapturerBackBuffer"));
-			FTexture2DRHIRef Texture = RHICreateTexture2D(VideoConfig.Width, VideoConfig.Height, EPixelFormat::PF_B8G8R8A8, 1, 1,
-			                                                           TexCreate_Shared | TexCreate_RenderTargetable | TexCreate_UAV, ERHIAccess::CopyDest, CreateInfo);
+			FTextureRHIRef Texture = RHICreateTexture(Desc);
+
 			InputFrame->SetTexture((ID3D12Resource*)Texture->GetNativeResource(), [&, InputFrame](ID3D12Resource* NativeTexture) { BackBuffers.Remove(InputFrame); });
 			BackBuffers.Add(InputFrame, Texture);
 		}
diff --git a/Engine/Source/Runtime/Landscape/Private/LandscapeEditReadback.cpp b/Engine/Source/Runtime/Landscape/Private/LandscapeEditReadback.cpp
--- a/Engine/Source/Runtime/Landscape/Private/LandscapeEditReadback.cpp
+++ b/Engine/Source/Runtime/Landscape/Private/LandscapeEditReadback.cpp
@@ -55,12 +55,16 @@ bool InitTask_RenderThread(FLandscapeEditReadbackTaskImpl& Task)
 	{
 		Task.StagingTextures.SetNum(Task.NumMips);
 
-		FRHIResourceCreateInfo CreateInfo(TEXT("LandscapeEditReadbackTask"));
 		for (uint32 MipIndex = 0; MipIndex < Task.NumMips; ++MipIndex)
 		{
 			const int32 MipWidth = FMath::Max(Task.Size.X >> MipIndex, 1);
 			const int32 MipHeight = FMath::Max(Task.Size.Y >> MipIndex, 1);
-			Task.StagingTextures[MipIndex] = RHICreateTexture2D(MipWidth, MipHeight, Task.Format, 1, 1, TexCreate_CPUReadback, CreateInfo);
+
+			const FRHITextureCreateDesc Desc =
+				FRHITextureCreateDesc::Create2D(TEXT("LandscapeEditReadbackTask"), MipWidth, MipHeight, Task.Format)
+				.SetFlags(ETextureCreateFlags::CPUReadback);
+
+			Task.StagingTextures[MipIndex] = RHICreateTexture(Desc);
 		}
 
 	}
diff --git a/Engine/Source/Runtime/Landscape/Private/LandscapeEditResources.cpp b/Engine/Source/Runtime/Landscape/Private/LandscapeEditResources.cpp
--- a/Engine/Source/Runtime/Landscape/Private/LandscapeEditResources.cpp
+++ b/Engine/Source/Runtime/Landscape/Private/LandscapeEditResources.cpp
@@ -26,20 +26,21 @@ void FLandscapeTexture2DResource::InitRHI()
 {
 	FTextureResource::InitRHI();
 
-	FRHIResourceCreateInfo CreateInfo(TEXT("FLandscapeTexture2DResource"));
-	ETextureCreateFlags Flags = TexCreate_None;
+	FRHITextureCreateDesc Desc =
+		FRHITextureCreateDesc::Create2D(TEXT("FLandscapeTexture2DResource"), SizeX, SizeY, Format)
+		.SetNumMips(NumMips);
 
 	if (bCreateUAVs)
 	{
-		Flags |= TexCreate_RenderTargetable | TexCreate_UAV;
+		Desc.AddFlags(ETextureCreateFlags::RenderTargetable | ETextureCreateFlags::UAV);
 	}
 
 	if (bCreateSRV)
 	{
-		Flags |= TexCreate_ShaderResource;
+		Desc.AddFlags(ETextureCreateFlags::ShaderResource);
 	}
 
-	TextureRHI = RHICreateTexture2D(SizeX, SizeY, Format, NumMips, 1, Flags, CreateInfo);
+	TextureRHI = RHICreateTexture(Desc);
 
 	if (bCreateUAVs)
 	{
diff --git a/Engine/Source/Runtime/MediaAssets/Private/Misc/MediaTextureResource.cpp b/Engine/Source/Runtime/MediaAssets/Private/Misc/MediaTextureResource.cpp
--- a/Engine/Source/Runtime/MediaAssets/Private/Misc/MediaTextureResource.cpp
+++ b/Engine/Source/Runtime/MediaAssets/Private/Misc/MediaTextureResource.cpp
@@ -769,15 +769,12 @@ void FMediaTextureResource::ConvertSample(const TSharedPtr<IMediaTextureSample,
 			// create a new temp input render target if necessary
 			if (!InputTarget.IsValid() || (InputTarget->GetSizeXY() != SampleDim) || (InputTarget->GetFormat() != InputPixelFormat) || ((InputTarget->GetFlags() & InputCreateFlags) != InputCreateFlags) || (InputTarget->GetNumMips() != SampleNumMips))
 			{
-				FRHIResourceCreateInfo CreateInfo(TEXT("FMediaTextureResource"));
-				InputTarget = RHICreateTexture2D(
-					SampleDim.X,
-					SampleDim.Y,
-					InputPixelFormat,
-					SampleNumMips,
-					1,
-					InputCreateFlags,
-					CreateInfo);
+				const FRHITextureCreateDesc Desc =
+					FRHITextureCreateDesc::Create2D(TEXT("FMediaTextureResource"), SampleDim, InputPixelFormat)
+					.SetNumMips(SampleNumMips)
+					.SetFlags(InputCreateFlags);
+
+				InputTarget = RHICreateTexture(Desc);
 
 				UpdateResourceSize();
 			}
diff --git a/Engine/Source/Runtime/MovieSceneCapture/Private/FrameGrabber.cpp b/Engine/Source/Runtime/MovieSceneCapture/Private/FrameGrabber.cpp
--- a/Engine/Source/Runtime/MovieSceneCapture/Private/FrameGrabber.cpp
+++ b/Engine/Source/Runtime/MovieSceneCapture/Private/FrameGrabber.cpp
@@ -58,17 +58,11 @@ void FViewportSurfaceReader::Resize(uint32 Width, uint32 Height)
 	ENQUEUE_RENDER_COMMAND(CreateCaptureFrameTexture)(
 		[Width, Height, This](FRHICommandListImmediate& RHICmdList)
 		{
-			FRHIResourceCreateInfo CreateInfo(TEXT("FViewportSurfaceReader_ReadbackTexture"));
-
-			This->ReadbackTexture = RHICreateTexture2D(
-				Width,
-				Height,
-				This->PixelFormat,
-				1,
-				1,
-				TexCreate_CPUReadback,
-				CreateInfo
-				);
+			const FRHITextureCreateDesc Desc =
+				FRHITextureCreateDesc::Create2D(TEXT("FViewportSurfaceReader_ReadbackTexture"), Width, Height, This->PixelFormat)
+				.SetFlags(ETextureCreateFlags::CPUReadback);
+
+			This->ReadbackTexture = RHICreateTexture(Desc);
 		});
 }
 
diff --git a/Engine/Source/Runtime/RHI/Private/RHIGPUReadback.cpp b/Engine/Source/Runtime/RHI/Private/RHIGPUReadback.cpp
--- a/Engine/Source/Runtime/RHI/Private/RHIGPUReadback.cpp
+++ b/Engine/Source/Runtime/RHI/Private/RHIGPUReadback.cpp
@@ -140,10 +140,14 @@ void FRHIGPUTextureReadback::EnqueueCopyInternal(FRHICommandList& RHICmdList, FR
 		if (!DestinationStagingTexture)
 		{
 			FIntVector TextureSize = SourceTexture->GetSizeXYZ();
-
 			FString FenceName = Fence->GetFName().ToString();
-			FRHIResourceCreateInfo CreateInfo(*FenceName);
-			DestinationStagingTexture = RHICreateTexture2D(TextureSize.X, TextureSize.Y, SourceTexture->GetFormat(), 1, 1, TexCreate_CPUReadback | TexCreate_HideInVisualizeTexture, ERHIAccess::CPURead, CreateInfo);
+
+			const FRHITextureCreateDesc Desc =
+				FRHITextureCreateDesc::Create2D(*FenceName, TextureSize.X, TextureSize.Y, SourceTexture->GetFormat())
+				.SetFlags(ETextureCreateFlags::CPUReadback | ETextureCreateFlags::HideInVisualizeTexture)
+				.SetInitialState(ERHIAccess::CPURead);
+
+			DestinationStagingTexture = RHICreateTexture(Desc);
 		}
 
 		// Transfer memory GPU -> CPU
diff --git a/Engine/Source/Runtime/RHI/Public/RHICommandList.h b/Engine/Source/Runtime/RHI/Public/RHICommandList.h
--- a/Engine/Source/Runtime/RHI/Public/RHICommandList.h
+++ b/Engine/Source/Runtime/RHI/Public/RHICommandList.h
@@ -5421,7 +5421,7 @@ FORCEINLINE FTextureRHIRef RHICreateTexture(const FRHITextureCreateDesc& CreateD
 	return GDynamicRHI->RHICreateTexture_RenderThread(RHICmdList, CreateDesc);
 }
 
-//UE_DEPRECATED(5.1, "FRHITexture2D is deprecated, please use RHICreateTexture(const FRHITextureCreateDesc&).")
+UE_DEPRECATED(5.1, "FRHITexture2D is deprecated, please use RHICreateTexture(const FRHITextureCreateDesc&).")
 FORCEINLINE FTexture2DRHIRef RHICreateTexture2D(uint32 SizeX, uint32 SizeY, uint8 Format, uint32 NumMips, uint32 NumSamples, ETextureCreateFlags Flags, ERHIAccess ResourceState, const FRHIResourceCreateInfo& CreateInfo)
 {
 	return RHICreateTexture(
@@ -5439,7 +5439,7 @@ FORCEINLINE FTexture2DRHIRef RHICreateTexture2D(uint32 SizeX, uint32 SizeY, uint
 	);
 }
 
-//UE_DEPRECATED(5.1, "FRHITexture2D is deprecated, please use RHICreateTexture(const FRHITextureCreateDesc&).")
+UE_DEPRECATED(5.1, "FRHITexture2D is deprecated, please use RHICreateTexture(const FRHITextureCreateDesc&).")
 FORCEINLINE FTexture2DRHIRef RHICreateTextureExternal2D(uint32 SizeX, uint32 SizeY, uint8 Format, uint32 NumMips, uint32 NumSamples, ETextureCreateFlags Flags, ERHIAccess ResourceState, const  FRHIResourceCreateInfo& CreateInfo)
 {
 PRAGMA_DISABLE_DEPRECATION_WARNINGS
@@ -5527,7 +5527,7 @@ FORCEINLINE FTextureCubeRHIRef RHICreateTextureCubeArray(uint32 Size, uint32 Arr
 	);
 }
 
-//UE_DEPRECATED(5.1, "FRHITexture2D is deprecated, please use RHICreateTexture(const FRHITextureCreateDesc&).")
+UE_DEPRECATED(5.1, "FRHITexture2D is deprecated, please use RHICreateTexture(const FRHITextureCreateDesc&).")
 FORCEINLINE FTexture2DRHIRef RHICreateTexture2D(uint32 SizeX, uint32 SizeY, uint8 Format, uint32 NumMips, uint32 NumSamples, ETextureCreateFlags Flags, const FRHIResourceCreateInfo& CreateInfo)
 {
 PRAGMA_DISABLE_DEPRECATION_WARNINGS
@@ -5535,7 +5535,7 @@ PRAGMA_DISABLE_DEPRECATION_WARNINGS
 PRAGMA_ENABLE_DEPRECATION_WARNINGS
 }
 
-//UE_DEPRECATED(5.1, "FRHITexture2D is deprecated, please use RHICreateTexture(const FRHITextureCreateDesc&).")
+UE_DEPRECATED(5.1, "FRHITexture2D is deprecated, please use RHICreateTexture(const FRHITextureCreateDesc&).")
 FORCEINLINE FTexture2DRHIRef RHICreateTextureExternal2D(uint32 SizeX, uint32 SizeY, uint8 Format, uint32 NumMips, uint32 NumSamples, ETextureCreateFlags Flags, const FRHIResourceCreateInfo& CreateInfo)
 {
 PRAGMA_DISABLE_DEPRECATION_WARNINGS
diff --git a/Engine/Source/Runtime/RHI/Public/RHIUtilities.h b/Engine/Source/Runtime/RHI/Public/RHIUtilities.h
--- a/Engine/Source/Runtime/RHI/Public/RHIUtilities.h
+++ b/Engine/Source/Runtime/RHI/Public/RHIUtilities.h
@@ -365,20 +365,17 @@ struct FTextureReadBuffer2D
 		Release();
 	}
 
-	const static ETextureCreateFlags DefaultTextureInitFlag = TexCreate_ShaderResource;
+	const static ETextureCreateFlags DefaultTextureInitFlag = ETextureCreateFlags::ShaderResource;
 	void Initialize(const TCHAR* InDebugName, const uint32 BytesPerElement, const uint32 SizeX, const uint32 SizeY, const EPixelFormat Format, ETextureCreateFlags Flags = DefaultTextureInitFlag, FResourceBulkDataInterface* InBulkData = nullptr)
 	{
 		NumBytes = SizeX * SizeY * BytesPerElement;
 
-		FRHIResourceCreateInfo CreateInfo(InDebugName);
-		CreateInfo.BulkData = InBulkData;
-		Buffer = RHICreateTexture2D(
-			SizeX, SizeY, UE_PIXELFORMAT_TO_UINT8(Format), //PF_R32_FLOAT,
-			/*NumMips=*/ 1,
-			1,
-			Flags,
-			/*BulkData=*/ CreateInfo);
+		const FRHITextureCreateDesc Desc =
+			FRHITextureCreateDesc::Create2D(InDebugName, SizeX, SizeY, Format)
+			.SetFlags(Flags)
+			.SetBulkData(InBulkData);
 
+		Buffer = RHICreateTexture(Desc);
 		
 		SRV = RHICreateShaderResourceView(Buffer, 0);
 	}
diff --git a/Engine/Source/Runtime/RenderCore/Private/DumpGPU.cpp b/Engine/Source/Runtime/RenderCore/Private/DumpGPU.cpp
--- a/Engine/Source/Runtime/RenderCore/Private/DumpGPU.cpp
+++ b/Engine/Source/Runtime/RenderCore/Private/DumpGPU.cpp
@@ -895,17 +895,13 @@ public:
 					SubresourceByteSize*= (4 / GPixelFormats[PreprocessedPixelFormat].BlockBytes);
 					PreprocessedPixelFormat = PF_R32_UINT;
 				}
-						
+
 				{
-					FRHIResourceCreateInfo CreateInfo(TEXT("DumpGPU.PreprocessTexture"));
-					StagingSrcTexture = RHICreateTexture2D(
-						SubresourceDumpDesc.SubResourceExtent.X,
-						SubresourceDumpDesc.SubResourceExtent.Y,
-						(uint8)PreprocessedPixelFormat,
-						/* NumMips = */ 1,
-						/* NumSamples = */ 1,
-						TexCreate_UAV | TexCreate_ShaderResource | TexCreate_HideInVisualizeTexture,
-						CreateInfo);
+					const FRHITextureCreateDesc Desc =
+						FRHITextureCreateDesc::Create2D(TEXT("DumpGPU.PreprocessTexture"), SubresourceDumpDesc.SubResourceExtent, PreprocessedPixelFormat)
+						.SetFlags(ETextureCreateFlags::UAV | ETextureCreateFlags::ShaderResource | ETextureCreateFlags::HideInVisualizeTexture);
+
+					StagingSrcTexture = RHICreateTexture(Desc);
 				}
 
 				FUnorderedAccessViewRHIRef StagingOutput = RHICreateUnorderedAccessView(StagingSrcTexture, /* MipLevel = */ 0);
@@ -934,15 +930,11 @@ public:
 
 			// Copy the texture for CPU readback
 			{
-				FRHIResourceCreateInfo CreateInfo(TEXT("DumpGPU.StagingTexture"));
-				StagingTexture = RHICreateTexture2D(
-					SubresourceDumpDesc.SubResourceExtent.X,
-					SubresourceDumpDesc.SubResourceExtent.Y,
-					(uint8)PreprocessedPixelFormat,
-					/* NumMips = */ 1,
-					/* NumSamples = */ 1,
-					TexCreate_CPUReadback | TexCreate_HideInVisualizeTexture,
-					CreateInfo);
+				const FRHITextureCreateDesc Desc =
+					FRHITextureCreateDesc::Create2D(TEXT("DumpGPU.StagingTexture"), SubresourceDumpDesc.SubResourceExtent, PreprocessedPixelFormat)
+					.SetFlags(ETextureCreateFlags::CPUReadback | ETextureCreateFlags::HideInVisualizeTexture);
+
+				StagingTexture = RHICreateTexture(Desc);
 
 				RHICmdList.Transition(FRHITransitionInfo(StagingTexture, ERHIAccess::Unknown, ERHIAccess::CopyDest));
 
diff --git a/Engine/Source/Runtime/RenderCore/Private/LongGPUTask.cpp b/Engine/Source/Runtime/RenderCore/Private/LongGPUTask.cpp
--- a/Engine/Source/Runtime/RenderCore/Private/LongGPUTask.cpp
+++ b/Engine/Source/Runtime/RenderCore/Private/LongGPUTask.cpp
@@ -13,8 +13,11 @@ int32 NumMeasuredIterationsToAchieve500ms = 0;
 
 void IssueScalableLongGPUTask(FRHICommandListImmediate& RHICmdList, int32 NumIteration /* = -1 by default */)
 {
-	FRHIResourceCreateInfo Info(TEXT("LongTaskRenderTarget"));
-	FTexture2DRHIRef LongTaskRenderTarget = RHICreateTexture2D(1920, 1080, PF_B8G8R8A8, 1, 1, TexCreate_RenderTargetable, Info);
+	const FRHITextureCreateDesc Desc =
+		FRHITextureCreateDesc::Create2D(TEXT("LongTaskRenderTarget"), 1920, 1080, PF_B8G8R8A8)
+		.SetFlags(ETextureCreateFlags::RenderTargetable);
+
+	FTextureRHIRef LongTaskRenderTarget = RHICreateTexture(Desc);
 
 	FRHIRenderPassInfo RPInfo(LongTaskRenderTarget, ERenderTargetActions::DontLoad_Store);
 	RHICmdList.Transition(FRHITransitionInfo(LongTaskRenderTarget, ERHIAccess::Unknown, ERHIAccess::RTV));
diff --git a/Engine/Source/Runtime/RenderCore/Private/RenderUtils.cpp b/Engine/Source/Runtime/RenderCore/Private/RenderUtils.cpp
--- a/Engine/Source/Runtime/RenderCore/Private/RenderUtils.cpp
+++ b/Engine/Source/Runtime/RenderCore/Private/RenderUtils.cpp
@@ -148,11 +148,14 @@ public:
 	{
 		// Create the texture RHI.  		
 		FBlackVolumeTextureResourceBulkDataInterface BlackTextureBulkData(FColor(R, G, B, A));
-		FRHIResourceCreateInfo CreateInfo(TEXT("ColoredTexture"), &BlackTextureBulkData);
-		ETextureCreateFlags CreateFlags = TexCreate_ShaderResource;
+
 		// BGRA typed UAV is unsupported per D3D spec, use RGBA here.
-		FTexture2DRHIRef Texture2D = RHICreateTexture2D(1, 1, PF_R8G8B8A8, 1, 1, CreateFlags, CreateInfo);
-		TextureRHI = Texture2D;
+		const FRHITextureCreateDesc Desc =
+			FRHITextureCreateDesc::Create2D(TEXT("ColoredTexture"), 1, 1, PF_R8G8B8A8)
+			.SetFlags(ETextureCreateFlags::ShaderResource)
+			.SetBulkData(&BlackTextureBulkData);
+
+		TextureRHI = RHICreateTexture(Desc);
 
 		// Create the sampler state RHI resource.
 		FSamplerStateInitializerRHI SamplerStateInitializer(SF_Point,AM_Wrap,AM_Wrap,AM_Wrap);
@@ -382,16 +385,20 @@ public:
 	{
 		// Create the texture RHI.
 		int32 TextureSize = 1 << (NumMips - 1);
-		FRHIResourceCreateInfo CreateInfo(TEXT("FMipColorTexture"));
-		FTexture2DRHIRef Texture2D = RHICreateTexture2D(TextureSize,TextureSize,PF_B8G8R8A8,NumMips,1,TexCreate_ShaderResource,CreateInfo);
-		TextureRHI = Texture2D;
+
+		const FRHITextureCreateDesc Desc =
+			FRHITextureCreateDesc::Create2D(TEXT("FMipColorTexture"), TextureSize, TextureSize, PF_B8G8R8A8)
+			.SetNumMips(NumMips)
+			.SetFlags(ETextureCreateFlags::ShaderResource);
+
+		TextureRHI = RHICreateTexture(Desc);
 
 		// Write the contents of the texture.
 		uint32 DestStride;
 		int32 Size = TextureSize;
 		for ( int32 MipIndex=0; MipIndex < NumMips; ++MipIndex )
 		{
-			FColor* DestBuffer = (FColor*)RHILockTexture2D(Texture2D, MipIndex, RLM_WriteOnly, DestStride, false);
+			FColor* DestBuffer = (FColor*)RHILockTexture2D(TextureRHI, MipIndex, RLM_WriteOnly, DestStride, false);
 			for ( int32 Y=0; Y < Size; ++Y )
 			{
 				for ( int32 X=0; X < Size; ++X )
@@ -400,7 +407,7 @@ public:
 				}
 				DestBuffer += DestStride / sizeof(FColor);
 			}
-			RHIUnlockTexture2D(Texture2D, MipIndex, false);
+			RHIUnlockTexture2D(TextureRHI, MipIndex, false);
 			Size >>= 1;
 		}
 
@@ -594,16 +601,17 @@ public:
 	// FResource interface.
 	virtual void InitRHI() override
 	{
-		// Create the texture RHI.  		
-		FRHIResourceCreateInfo CreateInfo(TEXT("UintTexture"));
-		FTexture2DRHIRef Texture2D = RHICreateTexture2D(1, 1, Format, 1, 1, TexCreate_ShaderResource, CreateInfo);
-		TextureRHI = Texture2D;
+		const FRHITextureCreateDesc Desc =
+			FRHITextureCreateDesc::Create2D(TEXT("UintTexture"), 1, 1, Format)
+			.SetFlags(ETextureCreateFlags::ShaderResource);
+
+		TextureRHI = RHICreateTexture(Desc);
 
 		// Write the contents of the texture.
 		uint32 DestStride;
-		void* DestBuffer = RHILockTexture2D(Texture2D, 0, RLM_WriteOnly, DestStride, false);
+		void* DestBuffer = RHILockTexture2D(TextureRHI, 0, RLM_WriteOnly, DestStride, false);
 		WriteData(DestBuffer);
-		RHIUnlockTexture2D(Texture2D, 0, false);
+		RHIUnlockTexture2D(TextureRHI, 0, false);
 
 		// Create the sampler state RHI resource.
 		FSamplerStateInitializerRHI SamplerStateInitializer(SF_Point, AM_Wrap, AM_Wrap, AM_Wrap);
diff --git a/Engine/Source/Runtime/Renderer/Private/RayTracing/RayTracingIESLightProfiles.cpp b/Engine/Source/Runtime/Renderer/Private/RayTracing/RayTracingIESLightProfiles.cpp
--- a/Engine/Source/Runtime/Renderer/Private/RayTracing/RayTracingIESLightProfiles.cpp
+++ b/Engine/Source/Runtime/Renderer/Private/RayTracing/RayTracingIESLightProfiles.cpp
@@ -39,8 +39,11 @@ void FIESLightProfileResource::BuildIESLightProfilesTexture(FRHICommandListImmed
 
 	if (!DefaultTexture)
 	{
-		FRHIResourceCreateInfo CreateInfo(TEXT("RTDefaultIESProfile"));
-		DefaultTexture = RHICreateTexture2D(AllowedIESProfileWidth, 1, AllowedIESProfileFormat, 1, 1, TexCreate_ShaderResource | TexCreate_UAV, CreateInfo);
+		const FRHITextureCreateDesc Desc =
+			FRHITextureCreateDesc::Create2D(TEXT("RTDefaultIESProfile"), AllowedIESProfileWidth, 1, AllowedIESProfileFormat)
+			.SetFlags(ETextureCreateFlags::ShaderResource | ETextureCreateFlags::UAV);
+
+		DefaultTexture = RHICreateTexture(Desc);
 		FUnorderedAccessViewRHIRef UAV = RHICreateUnorderedAccessView(DefaultTexture, 0);
 
 		RHICmdList.Transition(FRHITransitionInfo(DefaultTexture, ERHIAccess::Unknown, ERHIAccess::UAVCompute));
@@ -50,8 +53,11 @@ void FIESLightProfileResource::BuildIESLightProfilesTexture(FRHICommandListImmed
 
 	if (!AtlasTexture || AtlasTexture->GetSizeY() != NewArraySize)
 	{
-		FRHIResourceCreateInfo CreateInfo(TEXT("RTIESProfileAtlas"));
-		AtlasTexture = RHICreateTexture2D(AllowedIESProfileWidth, NewArraySize, AllowedIESProfileFormat, 1, 1, TexCreate_ShaderResource | TexCreate_UAV, CreateInfo);
+		const FRHITextureCreateDesc Desc =
+			FRHITextureCreateDesc::Create2D(TEXT("RTIESProfileAtlas"), AllowedIESProfileWidth, NewArraySize, AllowedIESProfileFormat)
+			.SetFlags(ETextureCreateFlags::ShaderResource | ETextureCreateFlags::UAV);
+
+		AtlasTexture = RHICreateTexture(Desc);
 		AtlasUAV = RHICreateUnorderedAccessView(AtlasTexture, 0);
 	}
 
diff --git a/Engine/Source/Runtime/SlateRHIRenderer/Private/SlateRHIFontTexture.cpp b/Engine/Source/Runtime/SlateRHIRenderer/Private/SlateRHIFontTexture.cpp
--- a/Engine/Source/Runtime/SlateRHIRenderer/Private/SlateRHIFontTexture.cpp
+++ b/Engine/Source/Runtime/SlateRHIRenderer/Private/SlateRHIFontTexture.cpp
@@ -20,11 +20,20 @@ void FSlateFontTextureRHIResource::InitDynamicRHI()
 	if( Width > 0 && Height > 0 )
 	{
 		const EPixelFormat PixelFormat = GetRHIPixelFormat();
-		const ETextureCreateFlags TextCreateFlags = TexCreate_Dynamic | (bIsGrayscale ? TexCreate_None : TexCreate_SRGB);
 
 		check( !IsValidRef( ShaderResource) );
-		FRHIResourceCreateInfo CreateInfo(TEXT("FSlateFontTextureRHIResource"));
-		ShaderResource = RHICreateTexture2D( Width, Height, PixelFormat, 1, 1, TextCreateFlags, CreateInfo );
+
+		FRHITextureCreateDesc Desc =
+			FRHITextureCreateDesc::Create2D(TEXT("FSlateFontTextureRHIResource"), Width, Height, PixelFormat)
+			.SetFlags(ETextureCreateFlags::Dynamic);
+
+		if (!bIsGrayscale)
+		{
+			Desc.AddFlags(ETextureCreateFlags::SRGB);
+		}
+
+		ShaderResource = RHICreateTexture(Desc);
+
 		check( IsValidRef( ShaderResource ) );
 
 		// Also assign the reference to the FTextureResource variable so that the Engine can access it
diff --git a/Engine/Source/Runtime/WebBrowser/Private/Android/AndroidWebBrowserWidget.cpp b/Engine/Source/Runtime/WebBrowser/Private/Android/AndroidWebBrowserWidget.cpp
--- a/Engine/Source/Runtime/WebBrowser/Private/Android/AndroidWebBrowserWidget.cpp
+++ b/Engine/Source/Runtime/WebBrowser/Private/Android/AndroidWebBrowserWidget.cpp
@@ -314,15 +314,18 @@ void SAndroidWebBrowserWidget::Tick(const FGeometry& AllottedGeometry, const dou
 					FTextureRHIRef VideoTexture = PinnedJavaWebBrowser->GetVideoTexture();
 					if (VideoTexture == nullptr)
 					{
-						FRHIResourceCreateInfo CreateInfo(TEXT("VideoTexture"));
-						FIntPoint LocalSize = Params.Size;
+						const FIntPoint LocalSize = Params.Size;
 
-						VideoTexture = RHICreateTextureExternal2D(LocalSize.X, LocalSize.Y, PF_R8G8B8A8, 1, 1, TexCreate_None, CreateInfo);
+						const FRHITextureCreateDesc Desc =
+							FRHITextureCreateDesc::Create2D(TEXT("VideoTexture"), LocalSize, PF_R8G8B8A8)
+							.SetFlags(ETextureCreateFlags::External);
+
+						VideoTexture = RHICreateTexture(Desc);
 						PinnedJavaWebBrowser->SetVideoTexture(VideoTexture);
 
 						if (VideoTexture == nullptr)
 						{
-							UE_LOG(LogAndroid, Warning, TEXT("CreateTextureExternal2D failed!"));
+							UE_LOG(LogAndroid, Warning, TEXT("RHICreateTexture failed!"));
 							return;
 						}
 
diff --git a/Engine/Source/Runtime/WebBrowser/Private/IOS/IOSPlatformWebBrowser.cpp b/Engine/Source/Runtime/WebBrowser/Private/IOS/IOSPlatformWebBrowser.cpp
--- a/Engine/Source/Runtime/WebBrowser/Private/IOS/IOSPlatformWebBrowser.cpp
+++ b/Engine/Source/Runtime/WebBrowser/Private/IOS/IOSPlatformWebBrowser.cpp
@@ -158,15 +158,17 @@ class SIOSWebBrowserWidget : public SLeafWidget
 							FTextureRHIRef VideoTexture = [NativeWebBrowser GetVideoTexture];
 							if (VideoTexture == nullptr)
 							{
-								FRHIResourceCreateInfo CreateInfo(TEXT("SIOSWebBrowserWidget_VideoTexture"));
-								FIntPoint Size = Params.Size;
-								VideoTexture = RHICreateTextureExternal2D(Size.X, Size.Y, PF_R8G8B8A8, 1, 1, TexCreate_None, CreateInfo);
+								const FRHITextureCreateDesc Desc =
+									FRHITextureCreateDesc::Create2D(TEXT("SIOSWebBrowserWidget_VideoTexture"), Params.Size, PF_R8G8B8A8)
+									.SetFlags(ETextureCreateFlags::External);
+
+								VideoTexture = RHICreateTexture(Desc);
 								[NativeWebBrowser SetVideoTexture : VideoTexture];
 								//UE_LOG(LogIOS, Log, TEXT("NativeWebBrowser SetVideoTexture:VideoTexture!"));
 
 								if (VideoTexture == nullptr)
 								{
-									UE_LOG(LogIOS, Warning, TEXT("CreateTextureExternal2D failed!"));
+									UE_LOG(LogIOS, Warning, TEXT("RHICreateTexture failed!"));
 									return;
 								}
 
