[ApplAR]
* Fixed nonunity compile error

diff --git a/Engine/Plugins/Runtime/AR/AppleAR/AppleARKit/Source/AppleARKit/Private/AppleARKitVideoOverlay.h b/Engine/Plugins/Runtime/AR/AppleAR/AppleARKit/Source/AppleARKit/Private/AppleARKitVideoOverlay.h
--- a/Engine/Plugins/Runtime/AR/AppleAR/AppleARKit/Source/AppleARKit/Private/AppleARKitVideoOverlay.h
+++ b/Engine/Plugins/Runtime/AR/AppleAR/AppleARKit/Source/AppleARKit/Private/AppleARKitVideoOverlay.h
@@ -13,6 +13,7 @@
 
 #include "AppleARKitVideoOverlay.generated.h"
 
+class FRHICommandListImmediate;
 class UARTextureCameraImage;
 class UMaterialInstanceDynamic;
 class UAppleARKitOcclusionTexture;
