[AppleARKit]
* Another attempt at fixing non unity compile errors

diff --git a/Engine/Plugins/Runtime/AR/AppleAR/AppleARKit/Source/AppleARKit/Private/AppleARKitVideoOverlay.h b/Engine/Plugins/Runtime/AR/AppleAR/AppleARKit/Source/AppleARKit/Private/AppleARKitVideoOverlay.h
--- a/Engine/Plugins/Runtime/AR/AppleAR/AppleARKit/Source/AppleARKit/Private/AppleARKitVideoOverlay.h
+++ b/Engine/Plugins/Runtime/AR/AppleAR/AppleARKit/Source/AppleARKit/Private/AppleARKitVideoOverlay.h
@@ -13,10 +13,12 @@
 
 #include "AppleARKitVideoOverlay.generated.h"
 
+class FRHIBuffer;
 class FRHICommandListImmediate;
 class UARTextureCameraImage;
 class UMaterialInstanceDynamic;
 class UAppleARKitOcclusionTexture;
+typedef TRefCountPtr<FRHIBuffer> FBufferRHIRef;
 
 enum class EARKitOcclusionType : uint8
 {
