Fixing nonunity issues

diff --git a/Engine/Plugins/Media/ElectraUtil/Source/ElectraSamples/Private/ElectraTextureSampleAVF.cpp b/Engine/Plugins/Media/ElectraUtil/Source/ElectraSamples/Private/ElectraTextureSampleAVF.cpp
--- a/Engine/Plugins/Media/ElectraUtil/Source/ElectraSamples/Private/ElectraTextureSampleAVF.cpp
+++ b/Engine/Plugins/Media/ElectraUtil/Source/ElectraSamples/Private/ElectraTextureSampleAVF.cpp
@@ -12,6 +12,7 @@
 #include "DataDrivenShaderPlatformInfo.h"
 #include "RHIStaticStates.h"
 #include "MediaShaders.h"
+#include "Containers/ResourceArray.h"
 #else
 #include "Containers/Array.h"
 #endif
diff --git a/Engine/Plugins/Runtime/AR/AppleAR/AppleARKit/Source/AppleARKit/Private/AppleARKitSystem.cpp b/Engine/Plugins/Runtime/AR/AppleAR/AppleARKit/Source/AppleARKit/Private/AppleARKitSystem.cpp
--- a/Engine/Plugins/Runtime/AR/AppleAR/AppleARKit/Source/AppleARKit/Private/AppleARKitSystem.cpp
+++ b/Engine/Plugins/Runtime/AR/AppleAR/AppleARKit/Source/AppleARKit/Private/AppleARKitSystem.cpp
@@ -25,6 +25,7 @@
 #include "ARBlueprintLibrary.h"
 #include "ARKitGeoTrackingSupport.h"
 #include "RenderGraphBuilder.h"
+#include "DataDrivenShaderPlatformInfo.h"
 
 // For mesh occlusion
 #include "MRMeshComponent.h"
diff --git a/Engine/Source/Runtime/Engine/Private/Components/RectLightComponent.cpp b/Engine/Source/Runtime/Engine/Private/Components/RectLightComponent.cpp
--- a/Engine/Source/Runtime/Engine/Private/Components/RectLightComponent.cpp
+++ b/Engine/Source/Runtime/Engine/Private/Components/RectLightComponent.cpp
@@ -16,6 +16,7 @@
 #include "RHIUtilities.h"
 #include "GlobalShader.h"
 #include "ShaderParameterUtils.h"
+#include "DataDrivenShaderPlatformInfo.h"
 
 #include UE_INLINE_GENERATED_CPP_BY_NAME(RectLightComponent)
 
