Fix non-Unity build (FMaterialRenderProxy) (UE-173241).

diff --git a/Engine/Plugins/Runtime/AR/AppleAR/AppleARKit/Source/AppleARKit/Private/AppleARKitTextures.cpp b/Engine/Plugins/Runtime/AR/AppleAR/AppleARKit/Source/AppleARKit/Private/AppleARKitTextures.cpp
--- a/Engine/Plugins/Runtime/AR/AppleAR/AppleARKit/Source/AppleARKit/Private/AppleARKitTextures.cpp
+++ b/Engine/Plugins/Runtime/AR/AppleAR/AppleARKit/Source/AppleARKit/Private/AppleARKitTextures.cpp
@@ -13,6 +13,7 @@
 #include "HAL/RunnableThread.h"
 #include "GlobalShader.h"
 #include "TextureResource.h"
+#include "Materials/MaterialRenderProxy.h"
 
 #if PLATFORM_APPLE
 	#import <Metal/Metal.h>
