Replaced deprecated ARHitTestResult with raycasting.

diff --git a/Engine/Plugins/Runtime/AR/AppleAR/AppleARKit/Source/AppleARKit/Private/AppleARKitSystem.cpp b/Engine/Plugins/Runtime/AR/AppleAR/AppleARKit/Source/AppleARKit/Private/AppleARKitSystem.cpp
--- a/Engine/Plugins/Runtime/AR/AppleAR/AppleARKit/Source/AppleARKit/Private/AppleARKitSystem.cpp
+++ b/Engine/Plugins/Runtime/AR/AppleAR/AppleARKit/Source/AppleARKit/Private/AppleARKitSystem.cpp
@@ -705,7 +705,7 @@ TArray<FARTraceResult> FAppleARKitSystem::OnLineTraceTrackedObjects( const FVect
 {
 	const float WorldToMetersScale = GetWorldToMetersScale();
 	TArray<FARTraceResult> Results;
-	
+
 	// Sanity check
 	if (IsRunning())
 	{
@@ -744,22 +744,17 @@ TArray<FARTraceResult> FAppleARKitSystem::OnLineTraceTrackedObjects( const FVect
 				
 				// GEngine->AddOnScreenDebugMessage(-1, 5.f, FColor::Blue, FString::Printf(TEXT("Hit Test At Screen Position: x: %f, y: %f"), NormalizedImagePosition.X, NormalizedImagePosition.Y));
 
-				//
-				// TODO: Re-enable deprecation warnings after updating the following code to use raycasting.
-				//
-				//   - 'ARHitTestResult' is deprecated: first deprecated in iOS 14.0 - Use raycasting
-				//   - 'hitTest:types:' is deprecated: first deprecated in iOS 14.0 - Use [ARSession raycast:]
-				//
-				PRAGMA_DISABLE_DEPRECATION_WARNINGS
-
 				// First run hit test against existing planes with extents (converting & filtering results as we go)
 				if (!!(TraceChannels & EARLineTraceChannels::PlaneUsingExtent) || !!(TraceChannels & EARLineTraceChannels::PlaneUsingBoundaryPolygon))
 				{
 					// First run hit test against existing planes with extents (converting & filtering results as we go)
-					NSArray< ARHitTestResult* >* PlaneHitTestResults = [HitTestFrame hitTest:CGPointMake(NormalizedImagePosition.X, NormalizedImagePosition.Y) types:ARHitTestResultTypeExistingPlaneUsingExtent];
-					for ( ARHitTestResult* HitTestResult in PlaneHitTestResults )
+					ARRaycastQuery* RaycastQuery = [HitTestFrame raycastQueryFromPoint: CGPointMake(NormalizedImagePosition.X, NormalizedImagePosition.Y) allowingTarget: ARRaycastTarget::ARRaycastTargetExistingPlaneGeometry alignment: ARRaycastTargetAlignment::ARRaycastTargetAlignmentAny];
+					NSArray< ARRaycastResult* >* PlaneHitTestResults = [Session raycast: RaycastQuery];
+					for (ARRaycastResult* HitTestResult in PlaneHitTestResults )
 					{
-						const float UnrealHitDistance = HitTestResult.distance * WorldToMetersScale;
+						FVector CameraPos = FAppleARKitConversion::ToFTransform(HitTestFrame.camera.transform).GetTranslation();
+						FVector IntersectionPos = FAppleARKitConversion::ToFTransform(HitTestResult.worldTransform).GetTranslation();
+						const float UnrealHitDistance = FVector::Dist(CameraPos, IntersectionPos);
 						if ( IsHitInRange( UnrealHitDistance ) )
 						{
 							// Hit result has passed and above filtering, add it to the list
@@ -772,10 +767,13 @@ TArray<FARTraceResult> FAppleARKitSystem::OnLineTraceTrackedObjects( const FVect
 				// If there were no valid results, fall back to hit testing against one shot plane
 				if (!!(TraceChannels & EARLineTraceChannels::GroundPlane))
 				{
-					NSArray< ARHitTestResult* >* PlaneHitTestResults = [HitTestFrame hitTest:CGPointMake(NormalizedImagePosition.X, NormalizedImagePosition.Y) types:ARHitTestResultTypeEstimatedHorizontalPlane];
-					for ( ARHitTestResult* HitTestResult in PlaneHitTestResults )
+					ARRaycastQuery* RaycastQuery = [HitTestFrame raycastQueryFromPoint : CGPointMake(NormalizedImagePosition.X, NormalizedImagePosition.Y) allowingTarget : ARRaycastTarget::ARRaycastTargetExistingPlaneInfinite alignment : ARRaycastTargetAlignment::ARRaycastTargetAlignmentHorizontal];
+					NSArray< ARRaycastResult* >* PlaneHitTestResults = [Session raycast : RaycastQuery];
+					for (ARRaycastResult* HitTestResult in PlaneHitTestResults )
 					{
-						const float UnrealHitDistance = HitTestResult.distance * WorldToMetersScale;
+						FVector CameraPos = FAppleARKitConversion::ToFTransform(HitTestFrame.camera.transform).GetTranslation();
+						FVector IntersectionPos = FAppleARKitConversion::ToFTransform(HitTestResult.worldTransform).GetTranslation();
+						const float UnrealHitDistance = FVector::Dist(CameraPos, IntersectionPos);
 						if ( IsHitInRange( UnrealHitDistance ) )
 						{
 							// Hit result has passed and above filtering, add it to the list
@@ -790,10 +788,13 @@ TArray<FARTraceResult> FAppleARKitSystem::OnLineTraceTrackedObjects( const FVect
 				{
 					// GEngine->AddOnScreenDebugMessage(-1, 5.f, FColor::Red, FString::Printf(TEXT("No results for plane hit test - reverting to feature points"), NormalizedImagePosition.X, NormalizedImagePosition.Y));
 					
-					NSArray< ARHitTestResult* >* FeatureHitTestResults = [HitTestFrame hitTest:CGPointMake(NormalizedImagePosition.X, NormalizedImagePosition.Y) types:ARHitTestResultTypeFeaturePoint];
-					for ( ARHitTestResult* HitTestResult in FeatureHitTestResults )
+					ARRaycastQuery* RaycastQuery = [HitTestFrame raycastQueryFromPoint : CGPointMake(NormalizedImagePosition.X, NormalizedImagePosition.Y) allowingTarget : ARRaycastTarget::ARRaycastTargetEstimatedPlane alignment : ARRaycastTargetAlignment::ARRaycastTargetAlignmentAny];
+					NSArray< ARRaycastResult* >*PlaneHitTestResults = [Session raycast : RaycastQuery];
+					for (ARRaycastResult* HitTestResult in PlaneHitTestResults)
 					{
-						const float UnrealHitDistance = HitTestResult.distance * WorldToMetersScale;
+						FVector CameraPos = FAppleARKitConversion::ToFTransform(HitTestFrame.camera.transform).GetTranslation();
+						FVector IntersectionPos = FAppleARKitConversion::ToFTransform(HitTestResult.worldTransform).GetTranslation();
+						const float UnrealHitDistance = FVector::Dist(CameraPos, IntersectionPos);
 						if ( IsHitInRange( UnrealHitDistance ) )
 						{
 							// Hit result has passed and above filtering, add it to the list
@@ -802,8 +803,6 @@ TArray<FARTraceResult> FAppleARKitSystem::OnLineTraceTrackedObjects( const FVect
 						}
 					}
 				}
-
-				PRAGMA_ENABLE_DEPRECATION_WARNINGS
 			}
 		}
 #endif
