UE-169465 Removing ARPins does not seem to release corresponding ARKit Anchors
-Release the arkit anchor when the pin is removed.

diff --git a/Engine/Plugins/Runtime/AR/AppleAR/AppleARKit/Source/AppleARKit/Private/AppleARKitSystem.cpp b/Engine/Plugins/Runtime/AR/AppleAR/AppleARKit/Source/AppleARKit/Private/AppleARKitSystem.cpp
--- a/Engine/Plugins/Runtime/AR/AppleAR/AppleARKit/Source/AppleARKit/Private/AppleARKitSystem.cpp
+++ b/Engine/Plugins/Runtime/AR/AppleAR/AppleARKit/Source/AppleARKit/Private/AppleARKitSystem.cpp
@@ -949,6 +949,26 @@ UARPin* FAppleARKitSystem::OnPinComponent( USceneComponent* ComponentToPin, cons
 
 void FAppleARKitSystem::OnRemovePin(UARPin* PinToRemove)
 {
+#if SUPPORTS_ARKIT_1_0
+	// Release the underlying ARKit anchor for this pin.
+	{
+		ARAnchor* Anchor = static_cast<ARAnchor*>(PinToRemove->GetNativeResource());
+		if (Anchor)
+		{
+			FScopeLock ScopeLock(&AnchorsLock);
+
+			const auto Guid = FAppleARKitConversion::ToFGuid(Anchor.identifier);
+			if (auto Record = AllAnchors.Find(Guid))
+			{
+				ARAnchor* SavedAnchor = *Record;
+				check(SavedAnchor == Anchor); // SavedAnchor should be the same as Anchor.
+				[SavedAnchor release] ;
+				AllAnchors.Remove(Guid);
+			}
+		}
+	}
+#endif
+
 	Pins.RemoveSingleSwap(PinToRemove);
 }
 
