Changing SetDrawRectangleParameters to use batched parameters.
Changing users of SetDrawRectangleParameters to fully use batched parameters.

diff --git a/Engine/Plugins/Experimental/RemoteSession/Source/RemoteSession/Private/Channels/RemoteSessionARCameraChannel.cpp b/Engine/Plugins/Experimental/RemoteSession/Source/RemoteSession/Private/Channels/RemoteSessionARCameraChannel.cpp
--- a/Engine/Plugins/Experimental/RemoteSession/Source/RemoteSession/Private/Channels/RemoteSessionARCameraChannel.cpp
+++ b/Engine/Plugins/Experimental/RemoteSession/Source/RemoteSession/Private/Channels/RemoteSessionARCameraChannel.cpp
@@ -93,10 +93,10 @@ public:
 		: FPostProcessMaterialShader(Initializer)
 	{}
 
-	void SetParameters(FRHICommandList& RHICmdList, const FSceneView& View)
+	void SetParameters(FRHIBatchedShaderParameters& BatchedParameters, const FSceneView& View)
 	{
-		FRHIVertexShader* ShaderRHI = RHICmdList.GetBoundVertexShader();
-		FMaterialShader::SetViewParameters(RHICmdList, ShaderRHI, View, View.ViewUniformBuffer);
+		UE::Renderer::PostProcess::SetDrawRectangleParameters(BatchedParameters, this, View);
+		FMaterialShader::SetViewParameters(BatchedParameters, View, View.ViewUniformBuffer);
 	}
 };
 
@@ -119,12 +119,11 @@ public:
 		: FPostProcessMaterialShader(Initializer)
 	{}
 
-	void SetParameters(FRHICommandList& RHICmdList, const FSceneView& View, const FMaterialRenderProxy* MaterialProxy)
+	void SetParameters(FRHIBatchedShaderParameters& BatchedParameters, const FSceneView& View, const FMaterialRenderProxy* MaterialProxy)
 	{
-		FRHIPixelShader* ShaderRHI = RHICmdList.GetBoundPixelShader();
 		const FMaterial& Material = MaterialProxy->GetMaterialWithFallback(View.GetFeatureLevel(), MaterialProxy);
-		FMaterialShader::SetViewParameters(RHICmdList, ShaderRHI, View, View.ViewUniformBuffer);
-		FMaterialShader::SetParameters(RHICmdList, ShaderRHI, MaterialProxy, Material, View);
+		FMaterialShader::SetViewParameters(BatchedParameters, View, View.ViewUniformBuffer);
+		FMaterialShader::SetParameters(BatchedParameters, MaterialProxy, Material, View);
 	}
 };
 
@@ -295,10 +294,8 @@ void FARCameraSceneViewExtension::RenderARCamera_RenderThread(FRDGBuilder& Graph
 
 		SetGraphicsPipelineState(RHICmdList, GraphicsPSOInit, 0);
 
-		UE::Renderer::PostProcess::SetDrawRectangleParameters(RHICmdList, VertexShader, InView);
-
-		VertexShader->SetParameters(RHICmdList, InView);
-		PixelShader->SetParameters(RHICmdList, InView, MaterialProxy);
+		SetShaderParametersLegacyVS(RHICmdList, VertexShader, InView);
+		SetShaderParametersLegacyPS(RHICmdList, PixelShader, InView, MaterialProxy);
 
 		RHICmdList.SetStreamSource(0, VertexBufferRHI, 0);
 		RHICmdList.DrawIndexedPrimitive(
diff --git a/Engine/Plugins/Runtime/AR/AppleAR/AppleARKit/Source/AppleARKit/Private/AppleARKitVideoOverlay.cpp b/Engine/Plugins/Runtime/AR/AppleAR/AppleARKit/Source/AppleARKit/Private/AppleARKitVideoOverlay.cpp
--- a/Engine/Plugins/Runtime/AR/AppleAR/AppleARKit/Source/AppleARKit/Private/AppleARKitVideoOverlay.cpp
+++ b/Engine/Plugins/Runtime/AR/AppleAR/AppleARKit/Source/AppleARKit/Private/AppleARKitVideoOverlay.cpp
@@ -170,10 +170,10 @@ public:
 		: TPostProcessMaterialShader<bIsMobileRenderer>(Initializer)
 	{}
 
-	void SetParameters(FRHICommandList& RHICmdList, const FSceneView View)
+	void SetParameters(FRHIBatchedShaderParameters& BatchedParameters, const FSceneView& View)
 	{
-		FRHIVertexShader* ShaderRHI = RHICmdList.GetBoundVertexShader();
-		TPostProcessMaterialShader<bIsMobileRenderer>::SetViewParameters(RHICmdList, ShaderRHI, View, View.ViewUniformBuffer);
+		UE::Renderer::PostProcess::SetDrawRectangleParameters(BatchedParameters, this, View);
+		TPostProcessMaterialShader<bIsMobileRenderer>::SetViewParameters(BatchedParameters, View, View.ViewUniformBuffer);
 	}
 };
 
@@ -201,11 +201,10 @@ public:
 		: TPostProcessMaterialShader<bIsMobileRenderer>(Initializer)
 	{}
 
-	void SetParameters(FRHICommandList& RHICmdList, const FSceneView View, const FMaterialRenderProxy* MaterialProxy, const FMaterial& Material)
+	void SetParameters(FRHIBatchedShaderParameters& BatchedParameters, const FSceneView& View, const FMaterialRenderProxy* MaterialProxy, const FMaterial& Material)
 	{
-		FRHIPixelShader* ShaderRHI = RHICmdList.GetBoundPixelShader();
-		TPostProcessMaterialShader<bIsMobileRenderer>::SetViewParameters(RHICmdList, ShaderRHI, View, View.ViewUniformBuffer);
-		TPostProcessMaterialShader<bIsMobileRenderer>::SetParameters(RHICmdList, ShaderRHI, MaterialProxy, Material, View);
+		TPostProcessMaterialShader<bIsMobileRenderer>::SetViewParameters(BatchedParameters, View, View.ViewUniformBuffer);
+		TPostProcessMaterialShader<bIsMobileRenderer>::SetParameters(BatchedParameters, MaterialProxy, Material, View);
 	}
 };
 
@@ -394,23 +393,21 @@ void FAppleARKitVideoOverlay::RenderVideoOverlayWithMaterial(FRHICommandListImme
 
 	SetGraphicsPipelineState(RHICmdList, GraphicsPSOInit, 0);
 
-	UE::Renderer::PostProcess::SetDrawRectangleParameters(RHICmdList, VertexShader, InView);
-
 	if (bIsMobileRenderer)
 	{
 		TShaderRef<FARKitCameraOverlayMobileVS> VertexShaderPtr(static_cast<FARKitCameraOverlayMobileVS*>(VertexShader.GetShader()), *MaterialShaderMap);
 		TShaderRef<FARKitCameraOverlayMobilePS> PixelShaderPtr(static_cast<FARKitCameraOverlayMobilePS*>(PixelShader.GetShader()), *MaterialShaderMap);
 
-		VertexShaderPtr->SetParameters(RHICmdList, InView);
-		PixelShaderPtr->SetParameters(RHICmdList, InView, MaterialProxy, CameraMaterial);
+		SetShaderParametersLegacyVS(RHICmdList, VertexShaderPtr, InView);
+		SetShaderParametersLegacyPS(RHICmdList, PixelShaderPtr, InView, MaterialProxy, CameraMaterial);
 	}
 	else
 	{
 		TShaderRef<FARKitCameraOverlayVS> VertexShaderPtr(static_cast<FARKitCameraOverlayVS*>(VertexShader.GetShader()), *MaterialShaderMap);
 		TShaderRef<FARKitCameraOverlayPS> PixelShaderPtr(static_cast<FARKitCameraOverlayPS*>(PixelShader.GetShader()), *MaterialShaderMap);
 
-		VertexShaderPtr->SetParameters(RHICmdList, InView);
-		PixelShaderPtr->SetParameters(RHICmdList, InView, MaterialProxy, CameraMaterial);
+		SetShaderParametersLegacyVS(RHICmdList, VertexShaderPtr, InView);
+		SetShaderParametersLegacyPS(RHICmdList, PixelShaderPtr, InView, MaterialProxy, CameraMaterial);
 	}
 	
 	if (OverlayVertexBufferRHI && IndexBufferRHI)
diff --git a/Engine/Plugins/Runtime/AR/Google/GoogleARCore/Source/GoogleARCoreRendering/Private/GoogleARCorePassthroughCameraRenderer.cpp b/Engine/Plugins/Runtime/AR/Google/GoogleARCore/Source/GoogleARCoreRendering/Private/GoogleARCorePassthroughCameraRenderer.cpp
--- a/Engine/Plugins/Runtime/AR/Google/GoogleARCore/Source/GoogleARCoreRendering/Private/GoogleARCorePassthroughCameraRenderer.cpp
+++ b/Engine/Plugins/Runtime/AR/Google/GoogleARCore/Source/GoogleARCoreRendering/Private/GoogleARCorePassthroughCameraRenderer.cpp
@@ -136,10 +136,10 @@ public:
 		: FPostProcessMaterialShader(Initializer)
 	{}
 
-	void SetParameters(FRHICommandList& RHICmdList, const FSceneView View)
+	void SetParameters(FRHIBatchedShaderParameters& BatchedParameters, const FSceneView& View)
 	{
-		FRHIVertexShader* ShaderRHI = RHICmdList.GetBoundVertexShader();
-		FMaterialShader::SetViewParameters(RHICmdList, ShaderRHI, View, View.ViewUniformBuffer);
+		UE::Renderer::PostProcess::SetDrawRectangleParameters(BatchedParameters, this, View);
+		FMaterialShader::SetViewParameters(BatchedParameters, View, View.ViewUniformBuffer);
 	}
 };
 
@@ -162,11 +162,10 @@ public:
 		: FPostProcessMaterialShader(Initializer)
 	{}
 
-	void SetParameters(FRHICommandList& RHICmdList, const FSceneView View, const FMaterialRenderProxy* MaterialProxy, const FMaterial& Material)
+	void SetParameters(FRHIBatchedShaderParameters& BatchedParameters, const FSceneView& View, const FMaterialRenderProxy* MaterialProxy, const FMaterial& Material)
 	{
-		FRHIPixelShader* ShaderRHI = RHICmdList.GetBoundPixelShader();
-		FMaterialShader::SetViewParameters(RHICmdList, ShaderRHI, View, View.ViewUniformBuffer);
-		FMaterialShader::SetParameters(RHICmdList, ShaderRHI, MaterialProxy, Material, View);
+		FMaterialShader::SetViewParameters(BatchedParameters, View, View.ViewUniformBuffer);
+		FMaterialShader::SetParameters(BatchedParameters, MaterialProxy, Material, View);
 	}
 };
 
@@ -228,10 +227,8 @@ void FGoogleARCorePassthroughCameraRenderer::RenderVideoOverlayWithMaterial(FRHI
 
 		SetGraphicsPipelineState(RHICmdList, GraphicsPSOInit, 0);
 
-		VertexShader->SetParameters(RHICmdList, InView);
-		PixelShader->SetParameters(RHICmdList, InView, MaterialProxy, CameraMaterial);
-
-		UE::Renderer::PostProcess::SetDrawRectangleParameters(RHICmdList, VertexShader, InView);
+		SetShaderParametersLegacyVS(RHICmdList, VertexShader, InView);
+		SetShaderParametersLegacyPS(RHICmdList, PixelShader, InView, MaterialProxy, CameraMaterial);
 
 		if (OverlayVertexBufferRHI.IsValid() && OverlayIndexBufferRHI.IsValid())
 		{
diff --git a/Engine/Source/Runtime/Renderer/Private/PostProcess/SceneFilterRendering.cpp b/Engine/Source/Runtime/Renderer/Private/PostProcess/SceneFilterRendering.cpp
--- a/Engine/Source/Runtime/Renderer/Private/PostProcess/SceneFilterRendering.cpp
+++ b/Engine/Source/Runtime/Renderer/Private/PostProcess/SceneFilterRendering.cpp
@@ -95,8 +95,8 @@ static void DoDrawRectangleFlagOverride(EDrawRectangleFlags& Flags)
 namespace UE::Renderer::PostProcess
 {
 	void SetDrawRectangleParameters(
-		FRHICommandList& RHICmdList,
-		const TShaderRef<FShader>& VertexShader,
+		FRHIBatchedShaderParameters& BatchedParameters,
+		const FShader* VertexShader,
 		float X,
 		float Y,
 		float SizeX,
@@ -119,13 +119,13 @@ namespace UE::Renderer::PostProcess
 			1.0f / TargetSize.X, 1.0f / TargetSize.Y,
 			1.0f / TextureSize.X, 1.0f / TextureSize.Y);
 
-		SetUniformBufferParameterImmediate(RHICmdList, VertexShader.GetVertexShader(), VertexShader->GetUniformBufferParameter<FDrawRectangleParameters>(), Parameters);
+		SetUniformBufferParameterImmediate(BatchedParameters, VertexShader->GetUniformBufferParameter<FDrawRectangleParameters>(), Parameters);
 	}
 
-	void SetDrawRectangleParameters(FRHICommandList& RHICmdList, const TShaderRef<FShader>& VertexShader, const FIntPoint& ViewSize)
+	void SetDrawRectangleParameters(FRHIBatchedShaderParameters& BatchedParameters, const FShader* VertexShader, const FIntPoint& ViewSize)
 	{
 		SetDrawRectangleParameters(
-			RHICmdList,
+			BatchedParameters,
 			VertexShader,
 			0.0f, 0.0f,
 			ViewSize.X, ViewSize.Y,
@@ -136,9 +136,9 @@ namespace UE::Renderer::PostProcess
 		);
 	}
 
-	void SetDrawRectangleParameters(FRHICommandList& RHICmdList, const TShaderRef<FShader>& VertexShader, const FSceneView& View)
+	void SetDrawRectangleParameters(FRHIBatchedShaderParameters& BatchedParameters, const FShader* VertexShader, const FSceneView& View)
 	{
-		SetDrawRectangleParameters(RHICmdList, VertexShader, View.UnconstrainedViewRect.Size());
+		SetDrawRectangleParameters(BatchedParameters, VertexShader, View.UnconstrainedViewRect.Size());
 	}
 }
 
@@ -168,7 +168,11 @@ inline void InternalDrawRectangle(
 		Flags = EDRF_Default;
 	}
 
-	UE::Renderer::PostProcess::SetDrawRectangleParameters(RHICmdList, VertexShader, X, Y, SizeX, SizeY, U, V, SizeU, SizeV, TargetSize, TextureSize);
+	{
+		FRHIBatchedShaderParameters& BatchedParameters = RHICmdList.GetScratchShaderParameters();
+		UE::Renderer::PostProcess::SetDrawRectangleParameters(BatchedParameters, VertexShader.GetShader(), X, Y, SizeX, SizeY, U, V, SizeU, SizeV, TargetSize, TextureSize);
+		RHICmdList.SetBatchedShaderParameters(VertexShader.GetVertexShader(), BatchedParameters);
+	}
 
 	if(Flags == EDRF_UseTesselatedIndexBuffer)
 	{
@@ -283,11 +287,15 @@ void DrawHmdMesh(
 	const TShaderRef<FShader>& VertexShader
 	)
 {
-	UE::Renderer::PostProcess::SetDrawRectangleParameters(
-		RHICmdList, VertexShader,
-		X, Y, SizeX, SizeY, U, V, SizeU, SizeV,
-		TargetSize, TextureSize
-	);
+	{
+		FRHIBatchedShaderParameters& BatchedParameters = RHICmdList.GetScratchShaderParameters();
+		UE::Renderer::PostProcess::SetDrawRectangleParameters(
+			BatchedParameters, VertexShader.GetShader(),
+			X, Y, SizeX, SizeY, U, V, SizeU, SizeV,
+			TargetSize, TextureSize
+		);
+		RHICmdList.SetBatchedShaderParameters(VertexShader.GetVertexShader(), BatchedParameters);
+	}
 
 	if (GEngine->XRSystem->GetHMDDevice())
 	{
diff --git a/Engine/Source/Runtime/Renderer/Public/PostProcess/DrawRectangle.h b/Engine/Source/Runtime/Renderer/Public/PostProcess/DrawRectangle.h
--- a/Engine/Source/Runtime/Renderer/Public/PostProcess/DrawRectangle.h
+++ b/Engine/Source/Runtime/Renderer/Public/PostProcess/DrawRectangle.h
@@ -15,8 +15,8 @@ template<typename ShaderType> using TShaderRef = TShaderRefBase<ShaderType, FSha
 namespace UE::Renderer::PostProcess
 {
 	RENDERER_API void SetDrawRectangleParameters(
-		FRHICommandList& RHICmdList,
-		const TShaderRef<FShader>& VertexShader,
+		FRHIBatchedShaderParameters& BatchedParameters,
+		const FShader* VertexShader,
 		float X,
 		float Y,
 		float SizeX,
@@ -28,8 +28,8 @@ namespace UE::Renderer::PostProcess
 		FIntPoint TargetSize,
 		FIntPoint TextureSize);
 
-	RENDERER_API void SetDrawRectangleParameters(FRHICommandList& RHICmdList, const TShaderRef<FShader>& VertexShader, const FIntPoint& ViewSize);
-	RENDERER_API void SetDrawRectangleParameters(FRHICommandList& RHICmdList, const TShaderRef<FShader>& VertexShader, const FSceneView& View);
+	RENDERER_API void SetDrawRectangleParameters(FRHIBatchedShaderParameters& BatchedParameters, const FShader* VertexShader, const FIntPoint& ViewSize);
+	RENDERER_API void SetDrawRectangleParameters(FRHIBatchedShaderParameters& BatchedParameters, const FShader* VertexShader, const FSceneView& View);
 
 	/**
 	 * Draws a quad with the given vertex positions and UVs in denormalized pixel/texel coordinates.
