Fixing ARKit overlay shaders by giving them shader parameters to bind, even though this code path never actually set those parameters to any values.

diff --git a/Engine/Plugins/Runtime/AR/AppleAR/AppleARKit/Source/AppleARKit/Private/AppleARKitVideoOverlay.cpp b/Engine/Plugins/Runtime/AR/AppleAR/AppleARKit/Source/AppleARKit/Private/AppleARKitVideoOverlay.cpp
--- a/Engine/Plugins/Runtime/AR/AppleAR/AppleARKit/Source/AppleARKit/Private/AppleARKitVideoOverlay.cpp
+++ b/Engine/Plugins/Runtime/AR/AppleAR/AppleARKit/Source/AppleARKit/Private/AppleARKitVideoOverlay.cpp
@@ -27,6 +27,8 @@
 #include "ARUtilitiesFunctionLibrary.h"
 #include "UObject/Package.h"
 #include "PostProcess/DrawRectangle.h"
+#include "ScreenPass.h"
+#include "PostProcess/PostProcessMaterialInputs.h"
 
 #if SUPPORTS_ARKIT_1_0
 	#include "IOSAppDelegate.h"
@@ -120,15 +122,20 @@ FAppleARKitVideoOverlay::~FAppleARKitVideoOverlay()
 	CVarDebugOverlayMode->SetOnChangedCallback({});
 }
 
+BEGIN_SHADER_PARAMETER_STRUCT(FARKitPostProcessMaterialParameters, )
+	SHADER_PARAMETER_STRUCT(FScreenPassTextureViewportParameters, PostProcessOutput)
+	SHADER_PARAMETER_STRUCT_ARRAY(FScreenPassTextureInput, PostProcessInput, [kPostProcessMaterialInputCountMax])
+	SHADER_PARAMETER_STRUCT_ARRAY(FScreenPassTextureInput, PathTracingPostProcessInput, [kPathTracingPostProcessMaterialInputCountMax])
+	SHADER_PARAMETER_SAMPLER(SamplerState, PostProcessInput_BilinearSampler)
+	RENDER_TARGET_BINDING_SLOTS()
+END_SHADER_PARAMETER_STRUCT()
+
 template <bool bIsMobileRenderer>
 class TPostProcessMaterialShader : public FMaterialShader
 {
 public:
-	TPostProcessMaterialShader() = default;
-	TPostProcessMaterialShader(const ShaderMetaType::CompiledShaderInitializerType& Initializer)
-		: FMaterialShader(Initializer)
-	{
-	}
+	using FParameters = FARKitPostProcessMaterialParameters;
+	SHADER_USE_PARAMETER_STRUCT_WITH_LEGACY_BASE(TPostProcessMaterialShader, FMaterialShader);
 
 	static bool ShouldCompilePermutation(const FMaterialShaderPermutationParameters& Parameters)
 	{
