One more missing FGCObject::AddReferencedObjects

#rnx
#horde 300570

diff --git a/Engine/Plugins/Runtime/AR/AppleAR/AppleARKit/Source/AppleARKit/Private/AppleARKitSystem.h b/Engine/Plugins/Runtime/AR/AppleAR/AppleARKit/Source/AppleARKit/Private/AppleARKitSystem.h
--- a/Engine/Plugins/Runtime/AR/AppleAR/AppleARKit/Source/AppleARKit/Private/AppleARKitSystem.h
+++ b/Engine/Plugins/Runtime/AR/AppleAR/AppleARKit/Source/AppleARKit/Private/AppleARKitSystem.h
@@ -61,6 +61,7 @@ public:
 private:
 	//~ FGCObject
 	virtual void AddReferencedObjects( FReferenceCollector& Collector ) override;
+	virtual FString GetReferencerName() const override { return "AppleARKitSystem"; }
 	//~ FGCObject
 protected:
 	//~IARSystemSupport
