Fixed mismatched copy constructor/assignment operator in FAppleARKitFrame.

diff --git a/Engine/Plugins/Runtime/AR/AppleAR/AppleARKit/Source/AppleARKit/Private/AppleARKitFrame.h b/Engine/Plugins/Runtime/AR/AppleAR/AppleARKit/Source/AppleARKit/Private/AppleARKitFrame.h
--- a/Engine/Plugins/Runtime/AR/AppleAR/AppleARKit/Source/AppleARKit/Private/AppleARKitFrame.h
+++ b/Engine/Plugins/Runtime/AR/AppleAR/AppleARKit/Source/AppleARKit/Private/AppleARKitFrame.h
@@ -19,7 +19,12 @@ struct APPLEARKIT_API FAppleARKitFrame
 	GENERATED_BODY()
 	
 	FAppleARKitFrame() = default;
+
+	// Added an assignment operator here because UStructs need them and Clang 16 expects either both
+	// or neither copy constructor and copy assignment operator to be implemented.  But it's pretty
+	// weird for a type to be copy assignable but not copy constructible.
 	FAppleARKitFrame(const FAppleARKitFrame& Other) = delete;
+	FAppleARKitFrame& operator=(const FAppleARKitFrame& Other) = default;
 	
 #if SUPPORTS_ARKIT_1_0
 	/** 
