Fixed mismatched copy constructor/assignment operator in FAppleARKitFrame.

diff --git a/Engine/Plugins/Runtime/AR/AppleAR/AppleARKit/Source/AppleARKit/Private/AppleARKitFrame.cpp b/Engine/Plugins/Runtime/AR/AppleAR/AppleARKit/Source/AppleARKit/Private/AppleARKitFrame.cpp
--- a/Engine/Plugins/Runtime/AR/AppleAR/AppleARKit/Source/AppleARKit/Private/AppleARKitFrame.cpp
+++ b/Engine/Plugins/Runtime/AR/AppleAR/AppleARKit/Source/AppleARKit/Private/AppleARKitFrame.cpp
@@ -170,6 +170,7 @@ FAppleARKitFrame::~FAppleARKitFrame()
 {
 	ReleaseResources();
 }
+#endif // SUPPORTS_ARKIT_1_0
 
 FAppleARKitFrame& FAppleARKitFrame::operator=( const FAppleARKitFrame& Other )
 {
@@ -177,12 +178,14 @@ FAppleARKitFrame& FAppleARKitFrame::operator=( const FAppleARKitFrame& Other )
 	{
 		return *this;
 	}
-	
+
+#if SUPPORTS_ARKIT_4_0
 	ReleaseResources();
-	
+
 	NativeFrame = Other.NativeFrame;
 	[NativeFrame retain];
-	
+#endif // SUPPORTS_ARKIT_1_0
+
 	// Member-wise copy
 	Timestamp = Other.Timestamp;
 	CapturedYImageSize = Other.CapturedYImageSize;
@@ -195,4 +198,3 @@ FAppleARKitFrame& FAppleARKitFrame::operator=( const FAppleARKitFrame& Other )
 
 	return *this;
 }
-#endif // SUPPORTS_ARKIT_1_0
