Fix iOS compile errors.

diff --git a/Engine/Plugins/Runtime/AR/AppleAR/AppleARKit/Source/AppleARKit/Private/AppleARKitSystem.cpp b/Engine/Plugins/Runtime/AR/AppleAR/AppleARKit/Source/AppleARKit/Private/AppleARKitSystem.cpp
--- a/Engine/Plugins/Runtime/AR/AppleAR/AppleARKit/Source/AppleARKit/Private/AppleARKitSystem.cpp
+++ b/Engine/Plugins/Runtime/AR/AppleAR/AppleARKit/Source/AppleARKit/Private/AppleARKitSystem.cpp
@@ -490,10 +490,10 @@ void FAppleARKitSystem::UpdateFrame()
 			if (FAppleARKitAvailability::SupportsARKit40() && GameThreadFrame->SceneDepth)
 			{
 				FImageBlurParams BlurParam = { SceneDepthBufferBlurAmount, SceneDepthBufferSizeScale };
-				UpdateCameraImageFromPixelBuffer(SceneDepthMap, GameThreadFrame->SceneDepth.depthMap, EARTextureType::SceneDepthMap, GameThreadFrame->Timestamp, PF_R32_FLOAT, kCGColorSpaceGenericRGBLinear, BlurParam);
+				UpdateCameraImageFromPixelBuffer(MutableView(SceneDepthMap), GameThreadFrame->SceneDepth.depthMap, EARTextureType::SceneDepthMap, GameThreadFrame->Timestamp, PF_R32_FLOAT, kCGColorSpaceGenericRGBLinear, BlurParam);
 				
 				// For some reason int8 texture needs to use the sRGB color space to ensure the values are intact after CIImage processing
-				UpdateCameraImageFromPixelBuffer(SceneDepthConfidenceMap, GameThreadFrame->SceneDepth.confidenceMap, EARTextureType::SceneDepthConfidenceMap, GameThreadFrame->Timestamp, PF_G8, kCGColorSpaceSRGB, BlurParam);
+				UpdateCameraImageFromPixelBuffer(MutableView(SceneDepthConfidenceMap), GameThreadFrame->SceneDepth.confidenceMap, EARTextureType::SceneDepthConfidenceMap, GameThreadFrame->Timestamp, PF_G8, kCGColorSpaceSRGB, BlurParam);
 				
 				if (FAppleARKitXRCamera* Camera = GetARKitXRCamera())
 				{
@@ -1491,7 +1491,7 @@ UARTrackedGeometry* FAppleARKitSystem::TryCreateTrackedGeometry(TSharedRef<FAppl
         {
             NewAnchorDebugName = FString::Printf(TEXT("IMG-%02d"), LastTrackedGeometry_DebugId++);
             UARTrackedImage* NewImage = NewObject<UARTrackedImage>();
-            UARCandidateImage** CandidateImage = CandidateImages.Find(AnchorData->DetectedAnchorName);
+            auto* CandidateImage = CandidateImages.Find(AnchorData->DetectedAnchorName);
             ensure(CandidateImage != nullptr);
             FVector2D PhysicalSize((*CandidateImage)->GetPhysicalWidth(), (*CandidateImage)->GetPhysicalHeight());
             NewImage->UpdateTrackedGeometry(ARComponent.ToSharedRef(), AnchorData->FrameNumber, AnchorData->Timestamp, AnchorData->Transform, AlignmentTransform, PhysicalSize, *CandidateImage);
@@ -1512,7 +1512,7 @@ UARTrackedGeometry* FAppleARKitSystem::TryCreateTrackedGeometry(TSharedRef<FAppl
         {
             NewAnchorDebugName = FString::Printf(TEXT("OBJ-%02d"), LastTrackedGeometry_DebugId++);
             UARTrackedObject* NewTrackedObject = NewObject<UARTrackedObject>();
-            UARCandidateObject** CandidateObject = CandidateObjects.Find(AnchorData->DetectedAnchorName);
+            auto* CandidateObject = CandidateObjects.Find(AnchorData->DetectedAnchorName);
             ensure(CandidateObject != nullptr);
             NewTrackedObject->UpdateTrackedGeometry(ARComponent.ToSharedRef(), AnchorData->FrameNumber, AnchorData->Timestamp, AnchorData->Transform, GetARCompositionComponent()->GetAlignmentTransform(), *CandidateObject);
             NewGeometry = NewTrackedObject;
@@ -1628,7 +1628,7 @@ bool FAppleARKitSystem::Run(UARSessionConfig* SessionConfig)
 		}
 		else
 		{
-			Configuration = FAppleARKitConversion::ToARConfiguration(SessionConfig, CandidateImages, ConvertedCandidateImages, CandidateObjects, InitialAlignmentTransform);
+			Configuration = FAppleARKitConversion::ToARConfiguration(SessionConfig, MutableView(CandidateImages), ConvertedCandidateImages, MutableView(CandidateObjects), InitialAlignmentTransform);
 		}
 
 		// Not all session types are supported by all devices
@@ -2424,7 +2424,7 @@ void FAppleARKitSystem::SessionDidUpdateAnchors_Internal( TSharedRef<FAppleARKit
             {
 				if (UARTrackedImage* ImageAnchor = Cast<UARTrackedImage>(FoundGeometry))
 				{
-					UARCandidateImage** CandidateImage = CandidateImages.Find(AnchorData->DetectedAnchorName);
+					auto* CandidateImage = CandidateImages.Find(AnchorData->DetectedAnchorName);
 					ensure(CandidateImage != nullptr);
 					FVector2D PhysicalSize((*CandidateImage)->GetPhysicalWidth(), (*CandidateImage)->GetPhysicalHeight());
 					ImageAnchor->UpdateTrackedGeometry(ARComponent.ToSharedRef(), AnchorData->FrameNumber, AnchorData->Timestamp, AnchorData->Transform, AlignmentTransform, PhysicalSize, *CandidateImage);
