ARKit: stop accessing camera texture resource on delegate thread

We can't acccess FTextureResource outside of the rendering thread. With ARKit, we receive the frame update on the delegate thread, where we grab the Metal texture to pass into our engine. We currently grab the FTextureResource on the delegate thread, but in newer UE versions, that is illegal (for good reason). So now we kick it to the render thread.

This was causing failures in Shipping builds because of how we handle thread-specific violations in Shipping vs non-Shipping builds. Non-shipping would 'succeed' in the face of a violation, where Shippin would correctly fail.

diff --git a/Engine/Plugins/Runtime/AR/AppleAR/AppleARKit/Source/AppleARKit/Private/AppleARKitTextures.cpp b/Engine/Plugins/Runtime/AR/AppleAR/AppleARKit/Source/AppleARKit/Private/AppleARKitTextures.cpp
--- a/Engine/Plugins/Runtime/AR/AppleAR/AppleARKit/Source/AppleARKit/Private/AppleARKitTextures.cpp
+++ b/Engine/Plugins/Runtime/AR/AppleAR/AppleARKit/Source/AppleARKit/Private/AppleARKitTextures.cpp
@@ -1120,23 +1120,25 @@ void UAppleARKitCameraVideoTexture::UpdateFrame(const FAppleARKitFrame& InFrame)
 #if SUPPORTS_ARKIT_1_0
 	if (InFrame.CapturedYImage && InFrame.CapturedCbCrImage)
 	{
-		if (auto VideoResource = static_cast<FARKitCameraVideoResource*>(GetResource()))
+		CVMetalTextureRef CapturedYImageCopy = InFrame.CapturedYImage;
+		CVMetalTextureRef CapturedCbCrImageCopy = InFrame.CapturedCbCrImage;
+		CFRetain(CapturedYImageCopy);
+		CFRetain(CapturedCbCrImageCopy);
+		const FIntPoint CapturedYImageSize = InFrame.CapturedYImageSize;
+		const FIntPoint CapturedCbCrImageSize = InFrame.CapturedCbCrImageSize;
+		const EDeviceScreenOrientation DeviceOrientation = FPlatformMisc::GetDeviceOrientation();
+
+		Size = CapturedYImageSize;
+
+		ENQUEUE_RENDER_COMMAND(UpdateVideoTexture)(
+			[this, CapturedYImageCopy, CapturedCbCrImageCopy, CapturedYImageSize, CapturedCbCrImageSize, DeviceOrientation](FRHICommandListImmediate& RHICmdList)
 		{
-			auto CapturedYImageCopy = InFrame.CapturedYImage;
-			auto CapturedCbCrImageCopy = InFrame.CapturedCbCrImage;
-			CFRetain(CapturedYImageCopy);
-			CFRetain(CapturedCbCrImageCopy);
-			const auto CapturedYImageSize = InFrame.CapturedYImageSize;
-			const auto CapturedCbCrImageSize = InFrame.CapturedCbCrImageSize;
-			const auto DeviceOrientation = FPlatformMisc::GetDeviceOrientation();
-			ENQUEUE_RENDER_COMMAND(UpdateVideoTexture)(
-				[VideoResource, CapturedYImageCopy, CapturedCbCrImageCopy, CapturedYImageSize, CapturedCbCrImageSize, DeviceOrientation](FRHICommandListImmediate& RHICmdList)
-			{
+			FARKitCameraVideoResource* VideoResource = static_cast<FARKitCameraVideoResource*>(GetResource());
+			if (VideoResource != nullptr)
+			{	
 				VideoResource->UpdateVideoTexture(RHICmdList, CapturedYImageCopy, CapturedYImageSize, CapturedCbCrImageCopy, CapturedCbCrImageSize, DeviceOrientation);
-			});
-			
-			Size = CapturedYImageSize;
-		}
+			}
+		});
 	}
 #endif
 }
