Fix conversion error in AppleARKit following changes in ARTextures (size now FVector2f instead of FVector2D).
See build health issue 318835.

diff --git a/Engine/Plugins/Runtime/AR/AppleAR/AppleARKit/Source/AppleARKit/Private/AppleARKitTextures.cpp b/Engine/Plugins/Runtime/AR/AppleAR/AppleARKit/Source/AppleARKit/Private/AppleARKitTextures.cpp
--- a/Engine/Plugins/Runtime/AR/AppleAR/AppleARKit/Source/AppleARKit/Private/AppleARKitTextures.cpp
+++ b/Engine/Plugins/Runtime/AR/AppleAR/AppleARKit/Source/AppleARKit/Private/AppleARKitTextures.cpp
@@ -754,7 +754,7 @@ void UAppleARKitOcclusionTexture::SetMetalTexture(float InTimestamp, id<MTLTextu
 			if (MetalTexture)
 			{
 				CFRetain(MetalTexture);
-				Size = FVector2D(MetalTexture.width, MetalTexture.height);
+				Size = FVector2f(MetalTexture.width, MetalTexture.height);
 			}
 		}
 	}
