ARKit: don't process invalid camera images

After pause/resume, the camera image returned by ARKit might not be valid to use. Ignore those, and set the UE texture size to 0.

diff --git a/Engine/Plugins/Runtime/AR/AppleAR/AppleARKit/Source/AppleARKit/Private/AppleARKitTextures.cpp b/Engine/Plugins/Runtime/AR/AppleAR/AppleARKit/Source/AppleARKit/Private/AppleARKitTextures.cpp
--- a/Engine/Plugins/Runtime/AR/AppleAR/AppleARKit/Source/AppleARKit/Private/AppleARKitTextures.cpp
+++ b/Engine/Plugins/Runtime/AR/AppleAR/AppleARKit/Source/AppleARKit/Private/AppleARKitTextures.cpp
@@ -1118,7 +1118,9 @@ void UAppleARKitCameraVideoTexture::Init()
 void UAppleARKitCameraVideoTexture::UpdateFrame(const FAppleARKitFrame& InFrame)
 {
 #if SUPPORTS_ARKIT_1_0
-	if (InFrame.CapturedYImage && InFrame.CapturedCbCrImage)
+    if ((InFrame.CapturedYImage && InFrame.CapturedCbCrImage) &&
+        (InFrame.CapturedYImageSize.GetMin() > 0) &&
+        (InFrame.CapturedCbCrImageSize.GetMin() > 0))
 	{
 		CVMetalTextureRef CapturedYImageCopy = InFrame.CapturedYImage;
 		CVMetalTextureRef CapturedCbCrImageCopy = InFrame.CapturedCbCrImage;
@@ -1140,7 +1142,11 @@ void UAppleARKitCameraVideoTexture::UpdateFrame(const FAppleARKitFrame& InFrame)
 			}
 		});
 	}
+    else
 #endif
+    {
+        Size = FVector2f::ZeroVector;
+    }
 }
 
 #if SUPPORTS_ARKIT_1_0
