PR #9723: make Live Link Face ios app can connect to the live link plugin on linux Platform

diff --git a/Engine/Plugins/Runtime/AR/AppleAR/AppleARKit/AppleARKit.uplugin b/Engine/Plugins/Runtime/AR/AppleAR/AppleARKit/AppleARKit.uplugin
--- a/Engine/Plugins/Runtime/AR/AppleAR/AppleARKit/AppleARKit.uplugin
+++ b/Engine/Plugins/Runtime/AR/AppleAR/AppleARKit/AppleARKit.uplugin
@@ -14,20 +14,20 @@
 	"IsBetaVersion": false,
 	"EnabledByDefault": false,
 	"SupportedTargetPlatforms" : [
-		"IOS", "Win64", "Mac"
+		"IOS", "Win64", "Mac","Linux"
 	],
 	"Modules": [
 		{
 			"Name": "AppleARKit",
 			"Type": "Runtime",
 			"LoadingPhase": "PostConfigInit",
-			"PlatformAllowList" : [ "Mac", "IOS", "Win64" ]
+			"PlatformAllowList" : [ "Mac", "IOS", "Win64","Linux" ]
 		},
 		{
 			"Name": "AppleARKitPoseTrackingLiveLink",
 			"Type": "Runtime",
 			"LoadingPhase": "PostConfigInit",
-			"PlatformAllowList" : [ "Mac", "IOS", "Win64" ]
+			"PlatformAllowList" : [ "Mac", "IOS", "Win64" ,"Linux"]
 
 		}
 	],
diff --git a/Engine/Plugins/Runtime/AR/AppleAR/AppleARKitFaceSupport/AppleARKitFaceSupport.uplugin b/Engine/Plugins/Runtime/AR/AppleAR/AppleARKitFaceSupport/AppleARKitFaceSupport.uplugin
--- a/Engine/Plugins/Runtime/AR/AppleAR/AppleARKitFaceSupport/AppleARKitFaceSupport.uplugin
+++ b/Engine/Plugins/Runtime/AR/AppleAR/AppleARKitFaceSupport/AppleARKitFaceSupport.uplugin
@@ -14,7 +14,7 @@
 	"IsBetaVersion": false,
 	"EnabledByDefault": false,
 	"SupportedTargetPlatforms" : [
-		"IOS", "Win64", "Mac"
+		"IOS", "Win64", "Mac","Linux"
 	],
 	"Modules":
 	[
@@ -22,7 +22,7 @@
 			"Name": "AppleARKitFaceSupport",
 			"Type": "Runtime",
 			"LoadingPhase": "PostConfigInit",
-			"PlatformAllowList" : [ "Mac", "IOS", "Win64" ]
+			"PlatformAllowList" : [ "Mac", "IOS", "Win64" ,"Linux"]
 		}
 	],
     "Plugins":
