Remove cyclic include
#rnx

diff --git a/Engine/Plugins/Runtime/AR/AppleAR/AppleARKit/Source/AppleARKit/Private/AppleARKitModule.h b/Engine/Plugins/Runtime/AR/AppleAR/AppleARKit/Source/AppleARKit/Private/AppleARKitModule.h
--- a/Engine/Plugins/Runtime/AR/AppleAR/AppleARKit/Source/AppleARKit/Private/AppleARKitModule.h
+++ b/Engine/Plugins/Runtime/AR/AppleAR/AppleARKit/Source/AppleARKit/Private/AppleARKitModule.h
@@ -6,7 +6,6 @@
 #include "IHeadMountedDisplayModule.h"
 #include "Features/IModularFeature.h"
 #include "AppleARKitSystem.h"
-#include "AppleARKitModule.h"
 
 class APPLEARKIT_API FAppleARKitModule : public IHeadMountedDisplayModule
 {
