UE-174645 ARKit anchor transforms trigger ensures when loaded from saved arspace
-Stripping scale from ARKit transforms while converting them to unreal coordinate system to avoid errors when they are not close enough to unit scale.

diff --git a/Engine/Plugins/Runtime/AR/AppleAR/AppleARKit/Source/AppleARKit/Public/AppleARKitConversion.h b/Engine/Plugins/Runtime/AR/AppleAR/AppleARKit/Source/AppleARKit/Public/AppleARKitConversion.h
--- a/Engine/Plugins/Runtime/AR/AppleAR/AppleARKit/Source/AppleARKit/Public/AppleARKitConversion.h
+++ b/Engine/Plugins/Runtime/AR/AppleAR/AppleARKit/Source/AppleARKit/Public/AppleARKitConversion.h
@@ -71,7 +71,7 @@ struct APPLEARKIT_API FAppleARKitConversion
 	 * Convert's an ARKit 'Y up' 'right handed' coordinate system transform to Unreal's 'Z up' 
 	 * 'left handed' coordinate system.
 	 *
-	 * Ignores scale.
+	 * Scale is removed because some ARKit transforms are less normalized than unreal requires.
 	 */
 	static FORCEINLINE FTransform ToFTransform(const matrix_float4x4& RawYUpMatrix, const FRotator& AdjustBy = FRotator::ZeroRotator)
 	{
@@ -82,6 +82,10 @@ struct APPLEARKIT_API FAppleARKitConversion
 			FPlane(RawYUpMatrix.columns[2][0], RawYUpMatrix.columns[2][1], RawYUpMatrix.columns[2][2], RawYUpMatrix.columns[2][3]),
 			FPlane(RawYUpMatrix.columns[3][0], RawYUpMatrix.columns[3][1], RawYUpMatrix.columns[3][2], RawYUpMatrix.columns[3][3]));
 
+		// ARKit, particularly when restoring a previously saved AR world map, sometimes has transforms who's normalization tolerance is not
+		// tight enough to avoid ensures or even failure that result in identity rotations from FQuat.  So we remove scaling here.
+		RawYUpFMatrix.RemoveScaling(UE_KINDA_SMALL_NUMBER / 100.0f);
+
 		// Extract & convert translation
 		FVector Translation = FVector( -RawYUpFMatrix.M[3][2], RawYUpFMatrix.M[3][0], RawYUpFMatrix.M[3][1] ) * ToUEScale();
 
