Changed ARKit to use objective c due to compile errors

diff --git a/Engine/Plugins/Runtime/AR/AppleAR/AppleARKit/Source/AppleARKit/Private/AppleARKitSystem.cpp b/Engine/Plugins/Runtime/AR/AppleAR/AppleARKit/Source/AppleARKit/Private/AppleARKitSystem.cpp
--- a/Engine/Plugins/Runtime/AR/AppleAR/AppleARKit/Source/AppleARKit/Private/AppleARKitSystem.cpp
+++ b/Engine/Plugins/Runtime/AR/AppleAR/AppleARKit/Source/AppleARKit/Private/AppleARKitSystem.cpp
@@ -1653,7 +1653,7 @@ bool FAppleARKitSystem::Run(UARSessionConfig* SessionConfig)
 			MTL::Device* Device = (MTL::Device*)GDynamicRHI->RHIGetNativeDevice();
 			check(Device);
 
-			CVReturn Return = CVMetalTextureCacheCreate(nullptr, nullptr, Device, nullptr, &MetalTextureCache);
+			CVReturn Return = CVMetalTextureCacheCreate(nullptr, nullptr, (__bridge id<MTLDevice>)Device, nullptr, &MetalTextureCache);
 			check(Return == kCVReturnSuccess);
 			check(MetalTextureCache);
 
diff --git a/Engine/Plugins/Runtime/AR/AppleAR/AppleARKit/Source/AppleARKit/Private/AppleARKitTrackable.h b/Engine/Plugins/Runtime/AR/AppleAR/AppleARKit/Source/AppleARKit/Private/AppleARKitTrackable.h
--- a/Engine/Plugins/Runtime/AR/AppleAR/AppleARKit/Source/AppleARKit/Private/AppleARKitTrackable.h
+++ b/Engine/Plugins/Runtime/AR/AppleAR/AppleARKit/Source/AppleARKit/Private/AppleARKitTrackable.h
@@ -19,7 +19,7 @@ public:
 	
 #if PLATFORM_IOS
 	/** Updates the current capture with the new metal texture. This will trigger a render resource update if the textures have changed */
-	void UpdateEnvironmentCapture(const TSharedRef<FARSupportInterface , ESPMode::ThreadSafe>& InTrackingSystem, uint32 FrameNumber, double Timestamp, const FTransform& InLocalToTrackingTransform, const FTransform& InAlignmentTransform, FVector InExtent, MTL::Texture* InMetalTexture);
+	void UpdateEnvironmentCapture(const TSharedRef<FARSupportInterface , ESPMode::ThreadSafe>& InTrackingSystem, uint32 FrameNumber, double Timestamp, const FTransform& InLocalToTrackingTransform, const FTransform& InAlignmentTransform, FVector InExtent, id<MTLTexture> InMetalTexture);
 #endif
 
 private:
diff --git a/Engine/Plugins/Runtime/AR/AppleAR/AppleARKit/Source/AppleARKit/Private/AppleARKitVideoOverlay.cpp b/Engine/Plugins/Runtime/AR/AppleAR/AppleARKit/Source/AppleARKit/Private/AppleARKitVideoOverlay.cpp
--- a/Engine/Plugins/Runtime/AR/AppleAR/AppleARKit/Source/AppleARKit/Private/AppleARKitVideoOverlay.cpp
+++ b/Engine/Plugins/Runtime/AR/AppleAR/AppleARKit/Source/AppleARKit/Private/AppleARKitVideoOverlay.cpp
@@ -243,7 +243,7 @@ void FAppleARKitVideoOverlay::UpdateOcclusionTextures(const FAppleARKitFrame& Fr
 		if (OcclusionType == EARKitOcclusionType::PersonSegmentation && NativeFrame &&
 			(NativeFrame.segmentationBuffer || NativeFrame.estimatedDepthData))
 		{
-			MTL::Device* Device = (MTL::Device*)GDynamicRHI->RHIGetNativeDevice();
+            id<MTLDevice> Device = (__bridge id<MTLDevice>)GDynamicRHI->RHIGetNativeDevice();
 			check(Device);
 
 			if (!MatteGenerator)
