Fix Unity Compile error in ARKit

diff --git a/Engine/Plugins/Runtime/AR/AppleAR/AppleARKit/Source/AppleARKit/Private/AppleARKitSystem.cpp b/Engine/Plugins/Runtime/AR/AppleAR/AppleARKit/Source/AppleARKit/Private/AppleARKitSystem.cpp
--- a/Engine/Plugins/Runtime/AR/AppleAR/AppleARKit/Source/AppleARKit/Private/AppleARKitSystem.cpp
+++ b/Engine/Plugins/Runtime/AR/AppleAR/AppleARKit/Source/AppleARKit/Private/AppleARKitSystem.cpp
@@ -1650,7 +1650,7 @@ bool FAppleARKitSystem::Run(UARSessionConfig* SessionConfig)
 		// Create MetalTextureCache
 		if (IsMetalPlatform(GMaxRHIShaderPlatform))
 		{
-			MTL::Device* Device = (MTL::Device*)GDynamicRHI->RHIGetNativeDevice();
+            id<MTLDevice> Device = (__bridge id<MTLDevice>)GDynamicRHI->RHIGetNativeDevice();
 			check(Device);
 
 			CVReturn Return = CVMetalTextureCacheCreate(nullptr, nullptr, (__bridge id<MTLDevice>)Device, nullptr, &MetalTextureCache);
