UE-193715 AppleARKitFaceSupport & AppleARKitPoseTrackingLiveLink reference the Private folder of the AppleARKit module
-Removed unused include and its private include path.

diff --git a/Engine/Plugins/Runtime/AR/AppleAR/AppleARKit/Source/AppleARKitPoseTrackingLiveLink/AppleARKitPoseTrackingLiveLink.build.cs b/Engine/Plugins/Runtime/AR/AppleAR/AppleARKit/Source/AppleARKitPoseTrackingLiveLink/AppleARKitPoseTrackingLiveLink.build.cs
--- a/Engine/Plugins/Runtime/AR/AppleAR/AppleARKit/Source/AppleARKitPoseTrackingLiveLink/AppleARKitPoseTrackingLiveLink.build.cs
+++ b/Engine/Plugins/Runtime/AR/AppleAR/AppleARKit/Source/AppleARKitPoseTrackingLiveLink/AppleARKitPoseTrackingLiveLink.build.cs
@@ -6,13 +6,7 @@ public class AppleARKitPoseTrackingLiveLink : ModuleRules
 {
 	public AppleARKitPoseTrackingLiveLink(ReadOnlyTargetRules Target) : base(Target)
 	{
-        //OptimizeCode = CodeOptimization.Never;
-		
-		PrivateIncludePaths.AddRange(new string[]
-		{
-			System.IO.Path.Combine(GetModuleDirectory("AppleARKit"), "Private"),
-		});
-			
+        //OptimizeCode = CodeOptimization.Never;	
 		
 		PublicDependencyModuleNames.AddRange(new string[]
 		{
diff --git a/Engine/Plugins/Runtime/AR/AppleAR/AppleARKit/Source/AppleARKitPoseTrackingLiveLink/Private/AppleARKitPoseTrackingLiveLinkModule.cpp b/Engine/Plugins/Runtime/AR/AppleAR/AppleARKit/Source/AppleARKitPoseTrackingLiveLink/Private/AppleARKitPoseTrackingLiveLinkModule.cpp
--- a/Engine/Plugins/Runtime/AR/AppleAR/AppleARKit/Source/AppleARKitPoseTrackingLiveLink/Private/AppleARKitPoseTrackingLiveLinkModule.cpp
+++ b/Engine/Plugins/Runtime/AR/AppleAR/AppleARKit/Source/AppleARKitPoseTrackingLiveLink/Private/AppleARKitPoseTrackingLiveLinkModule.cpp
@@ -1,7 +1,6 @@
 // Copyright Epic Games, Inc. All Rights Reserved.
 
 #include "AppleARKitPoseTrackingLiveLinkModule.h"
-#include "AppleARKitSystem.h"
 #include "Features/IModularFeature.h"
 #include "Features/IModularFeatures.h"
 #include "CoreGlobals.h"
