UE-215556 XR -ARKit - Tracked Planes are not attached to the real world objects when moving the camera around
-The FOVs with which we render the virtual content were wrong because aspect ratio was assumed to always be 16:9 which is never correct for recent iphones and never even close for portrait orientation.  One possible fix is to also override the aspect ratio, however that would not be hotfix complient.  Instead I've overrideen SetupViewProjectionMatrix so that AppleARKitSystem can build the projection matrix directly, which is consistent with ARCore and similar to how it is done for stereoscopic rendering.

diff --git a/Engine/Plugins/Runtime/AR/AppleAR/AppleARKit/Source/AppleARKit/Private/AppleARKitCamera.cpp b/Engine/Plugins/Runtime/AR/AppleAR/AppleARKit/Source/AppleARKit/Private/AppleARKitCamera.cpp
--- a/Engine/Plugins/Runtime/AR/AppleAR/AppleARKit/Source/AppleARKit/Private/AppleARKitCamera.cpp
+++ b/Engine/Plugins/Runtime/AR/AppleAR/AppleARKit/Source/AppleARKit/Private/AppleARKitCamera.cpp
@@ -10,6 +10,7 @@
 #include "Engine/EngineTypes.h"
 #include "Kismet/GameplayStatics.h"
 #include "Rendering/RenderingCommon.h"
+#include "SceneView.h"
 
 #if SUPPORTS_ARKIT_1_0
 
@@ -228,6 +229,28 @@ float FAppleARKitCamera::GetVerticalFieldOfViewForScreen( EAppleARKitBackgroundF
 	return GetVerticalFieldOfView();
 }
 
+void FAppleARKitCamera::GetViewProjectionMatrix(EDeviceScreenOrientation DeviceOrientation, FSceneViewProjectionData& InOutProjectionData) const
+{
+	// Use the global viewport size as the screen size
+	FVector2D ViewportSize;
+	GEngine->GameViewport->GetViewportSize( ViewportSize );
+	
+	float FOV = 0.0f;
+	if (DeviceOrientation == EDeviceScreenOrientation::Portrait || DeviceOrientation == EDeviceScreenOrientation::PortraitUpsideDown)
+	{
+		// Portrait
+		FOV = GetVerticalFieldOfViewForScreen(EAppleARKitBackgroundFitMode::Fill);
+	}
+	else
+	{
+		// Landscape
+		FOV = GetHorizontalFieldOfViewForScreen(EAppleARKitBackgroundFitMode::Fill);
+	}
+	
+	InOutProjectionData.ProjectionMatrix = FReversedZPerspectiveMatrix(FMath::DegreesToRadians(FOV) * 0.5f, ViewportSize.X, ViewportSize.Y, 10.0f);
+}
+
+
 FVector2D FAppleARKitCamera::GetImageCoordinateForScreenPosition( FVector2D ScreenPosition, EAppleARKitBackgroundFitMode BackgroundFitMode, float ScreenWidth, float ScreenHeight ) const
 {
 	const bool bIsInPortraitMode = ScreenWidth < ScreenHeight;
diff --git a/Engine/Plugins/Runtime/AR/AppleAR/AppleARKit/Source/AppleARKit/Private/AppleARKitCamera.h b/Engine/Plugins/Runtime/AR/AppleAR/AppleARKit/Source/AppleARKit/Private/AppleARKitCamera.h
--- a/Engine/Plugins/Runtime/AR/AppleAR/AppleARKit/Source/AppleARKit/Private/AppleARKitCamera.h
+++ b/Engine/Plugins/Runtime/AR/AppleAR/AppleARKit/Source/AppleARKit/Private/AppleARKitCamera.h
@@ -17,6 +17,8 @@
 
 #include "AppleARKitCamera.generated.h"
 
+struct FSceneViewProjectionData;
+
 
 enum class EAppleARKitBackgroundFitMode : uint8
 {
@@ -125,6 +127,11 @@ struct APPLEARKIT_API FAppleARKitCamera
 	/** Returns the effective vertical field of view for the screen dimensions and fit mode in those dimensions */
 	float GetVerticalFieldOfViewForScreen( EAppleARKitBackgroundFitMode BackgroundFitMode, float ScreenWidth, float ScreenHeight ) const;
 
+	/** Returns the Projection Matrix we need to render our virtual content with. */
+	void GetViewProjectionMatrix(EDeviceScreenOrientation DeviceOrientation, FSceneViewProjectionData& InOutProjectionData) const;
+	
+	
+
 	/** For the given screen position, returns the normalised capture image coordinates accounting for the fit mode of the image on screen */
 	FVector2D GetImageCoordinateForScreenPosition( FVector2D ScreenPosition, EAppleARKitBackgroundFitMode BackgroundFitMode ) const
 	{
diff --git a/Engine/Plugins/Runtime/AR/AppleAR/AppleARKit/Source/AppleARKit/Private/AppleARKitSystem.cpp b/Engine/Plugins/Runtime/AR/AppleAR/AppleARKit/Source/AppleARKit/Private/AppleARKitSystem.cpp
--- a/Engine/Plugins/Runtime/AR/AppleAR/AppleARKit/Source/AppleARKit/Private/AppleARKitSystem.cpp
+++ b/Engine/Plugins/Runtime/AR/AppleAR/AppleARKit/Source/AppleARKit/Private/AppleARKitSystem.cpp
@@ -148,7 +148,10 @@ private:
 	
 	virtual void SetupViewProjectionMatrix(FSceneViewProjectionData& InOutProjectionData) override
 	{
-		FDefaultXRCamera::SetupViewProjectionMatrix(InOutProjectionData);
+		if (ARKitSystem.GameThreadFrame.IsValid()) 
+		{
+			ARKitSystem.GameThreadFrame->Camera.GetViewProjectionMatrix(ARKitSystem.DeviceOrientation, InOutProjectionData);
+		}
 	}
 	
 	virtual void BeginRenderViewFamily(FSceneViewFamily& InViewFamily) override
