RHIBindDebugLabelName Cleanup
- Removing redundant uses of RHIBindDebugLabelName, FRHITexture::SetName, and FRHIBuffer::SetName
- Remaining uses of RHIBindDebugLabelName now all use the RHICommandList function.
- Deprecated the global RHIBindDebugLabelName functions.

diff --git a/Engine/Platforms/VisionOS/Plugins/Runtime/OpenXRVisionOS/Source/OXRVisionOS/Private/OXRVisionOSSwapchain.cpp b/Engine/Platforms/VisionOS/Plugins/Runtime/OpenXRVisionOS/Source/OXRVisionOS/Private/OXRVisionOSSwapchain.cpp
--- a/Engine/Platforms/VisionOS/Plugins/Runtime/OpenXRVisionOS/Source/OXRVisionOS/Private/OXRVisionOSSwapchain.cpp
+++ b/Engine/Platforms/VisionOS/Plugins/Runtime/OpenXRVisionOS/Source/OXRVisionOS/Private/OXRVisionOSSwapchain.cpp
@@ -114,6 +114,9 @@ FOXRVisionOSSwapchain::FOXRVisionOSSwapchain(const XrSwapchainCreateInfo* create
 				check(false);
 			}
 		}
+
+		FRHICommandListImmediate& RHICmdList = FRHICommandListImmediate::Get();
+
 		const FRHITextureCreateDesc BackBufferDesc =
 			FRHITextureCreateDesc::Create2D(isColor ? TEXT("OXRVisionOSSwapchain") : TEXT("OXRVisionOSSwapchainDepth"),
 											CreateInfo.width, CreateInfo.height,
@@ -125,13 +128,13 @@ FOXRVisionOSSwapchain::FOXRVisionOSSwapchain(const XrSwapchainCreateInfo* create
 		Images.Reserve(SwapchainLength);
 		for (int32 BufferIndex = 0; BufferIndex < SwapchainLength; ++BufferIndex)
 		{
-			FTextureRHIRef Tex = RHICreateTexture(BackBufferDesc);
-			Images.Emplace(Tex, EImageState::Released);
+			const FString Name = FString::Printf(TEXT("%s_SwapChain%d"), (CreateInfo.usageFlags & XR_SWAPCHAIN_USAGE_COLOR_ATTACHMENT_BIT) ? TEXT("Color") : TEXT("Depth"), BufferIndex);
 
-			// hacky way to set the debug label for these textures- this can go when everything is working, but it's helpful for debugging in the capture
-			FString Name = FString::Printf(TEXT("%s_SwapChain%d"), (CreateInfo.usageFlags & XR_SWAPCHAIN_USAGE_COLOR_ATTACHMENT_BIT) ? TEXT("Color") : TEXT("Depth"), BufferIndex);
-			FRHICommandListBase* Base = (FRHICommandListBase*)nullptr;
-			GDynamicRHI->RHIBindDebugLabelName(*Base, Tex, *Name);
+			FRHITextureCreateDesc CurrentBackBufferDesc = BackBufferDesc;
+			CurrentBackBufferDesc.DebugName = *Name;
+
+			FTextureRHIRef Tex = RHICmdList.CreateTexture(CurrentBackBufferDesc);
+			Images.Emplace(Tex, EImageState::Released);
 		}
 	}
 }
diff --git a/Engine/Plugins/Runtime/AR/AppleAR/AppleARKit/Source/AppleARKit/Private/AppleARKitTextures.cpp b/Engine/Plugins/Runtime/AR/AppleAR/AppleARKit/Source/AppleARKit/Private/AppleARKitTextures.cpp
--- a/Engine/Plugins/Runtime/AR/AppleAR/AppleARKit/Source/AppleARKit/Private/AppleARKitTextures.cpp
+++ b/Engine/Plugins/Runtime/AR/AppleAR/AppleARKit/Source/AppleARKit/Private/AppleARKitTextures.cpp
@@ -159,7 +159,7 @@ public:
 		
 		TextureRHI = RHICreateTexture(Desc);
 		
-		OnTextureUpdated();
+		OnTextureUpdated(RHICmdList);
 	}
 	
 	virtual void ReleaseRHI() override
@@ -177,9 +177,9 @@ public:
 		}
 		CFRetain(CameraImage);
 		ENQUEUE_RENDER_COMMAND(UpdateCameraImage)
-		([this, CameraImage, PixelFormat, ColorSpace, BlurParams](FRHICommandListImmediate&)
+		([this, CameraImage, PixelFormat, ColorSpace, BlurParams](FRHICommandListImmediate& RHICmdList)
 		{
-			UpdateCameraImage_RenderThread(CameraImage, PixelFormat, ColorSpace, BlurParams);
+			UpdateCameraImage_RenderThread(RHICmdList, CameraImage, PixelFormat, ColorSpace, BlurParams);
 			CFRelease(CameraImage);
 		});
 	}
@@ -195,13 +195,13 @@ public:
 		ENQUEUE_RENDER_COMMAND(UpdateMetalTextureResource)
 		([this, MetalTexture, PixelFormat, ColorSpace](FRHICommandListImmediate& RHICmdList)
 		{
-			UpdateMetalTexture_RenderThread(MetalTexture, PixelFormat, ColorSpace);
+			UpdateMetalTexture_RenderThread(RHICmdList, MetalTexture, PixelFormat, ColorSpace);
 			CFRelease(MetalTexture);
 		});
 	}
 	
 protected:
-	void ConditionalRecreateTexture(uint32 Width, uint32 Height, EPixelFormat PixelFormat)
+	void ConditionalRecreateTexture(FRHICommandListImmediate& RHICmdList, uint32 Width, uint32 Height, EPixelFormat PixelFormat)
 	{
 		if (Size.X != Width || Size.Y != Height)
 		{
@@ -209,7 +209,7 @@ protected:
 			Size.Y = Height;
 			
 			// Let go of the last texture
-			RHIUpdateTextureReference(Owner->TextureReference.TextureReferenceRHI, nullptr);
+			RHICmdList.UpdateTextureReference(Owner->TextureReference.TextureReferenceRHI, nullptr);
 			TextureRHI.SafeRelease();
 			
 			// Create the target texture that we'll update into
@@ -217,11 +217,11 @@ protected:
 				FRHITextureCreateDesc::Create2D(TEXT("FARKitTextureResource"), Size, PixelFormat)
 				.SetFlags(ETextureCreateFlags::ShaderResource | ETextureCreateFlags::UAV);
 
-			TextureRHI = RHICreateTexture(Desc);
+			TextureRHI = RHICmdList.CreateTexture(Desc);
 		}
 	}
 	
-	void UpdateCameraImage_RenderThread(CVPixelBufferRef CameraImage, EPixelFormat PixelFormat, const CFStringRef ColorSpace, FImageBlurParams BlurParams)
+	void UpdateCameraImage_RenderThread(FRHICommandListImmediate& RHICmdList, CVPixelBufferRef CameraImage, EPixelFormat PixelFormat, const CFStringRef ColorSpace, FImageBlurParams BlurParams)
 	{
 		SCOPED_AUTORELEASE_POOL;
 		CGColorSpaceRef ColorSpaceRef = CGColorSpaceCreateWithName(ColorSpace);
@@ -253,7 +253,7 @@ protected:
 		CGRect ImageExtent = RotatedImage.extent;
 		
 		// Don't reallocate the texture if the sizes match
-		ConditionalRecreateTexture(ImageExtent.size.width, ImageExtent.size.height, PixelFormat);
+		ConditionalRecreateTexture(RHICmdList, ImageExtent.size.width, ImageExtent.size.height, PixelFormat);
 		
 		// Get the underlying metal texture so we can render to it
 		id<MTLTexture> UnderlyingMetalTexture = (id<MTLTexture>)TextureRHI->GetNativeResource();
@@ -276,10 +276,10 @@ protected:
 			ensure(ProcessedTexture == UnderlyingMetalTexture);
 		}
 		
-		OnTextureUpdated();
+		OnTextureUpdated(RHICmdList);
 	}
 	
-	void UpdateMetalTexture_RenderThread(id<MTLTexture> MetalTexture, EPixelFormat PixelFormat, const CFStringRef ColorSpace)
+	void UpdateMetalTexture_RenderThread(FRHICommandListImmediate& RHICmdList, id<MTLTexture> MetalTexture, EPixelFormat PixelFormat, const CFStringRef ColorSpace)
 	{
 		SCOPED_AUTORELEASE_POOL;
 		
@@ -310,7 +310,7 @@ protected:
 		
 		if (!TextureRHI || DesiredSize != Size)
 		{
-			ConditionalRecreateTexture(DesiredSize.X, DesiredSize.Y, PixelFormat);
+			ConditionalRecreateTexture(RHICmdList, DesiredSize.X, DesiredSize.Y, PixelFormat);
 		}
 		
 		if (TextureRHI)
@@ -325,7 +325,7 @@ protected:
 		[Image release];
 		CGColorSpaceRelease(ColorSpaceRef);
 		
-		OnTextureUpdated();
+		OnTextureUpdated(RHICmdList);
 	}
 	
 	CIContext* GetImageContext()
@@ -340,11 +340,10 @@ protected:
 #endif // PLATFORM_APPLE
 	
 protected:
-	void OnTextureUpdated()
+	void OnTextureUpdated(FRHICommandListBase& RHICmdList)
 	{
-		TextureRHI->SetName(Owner->GetFName());
-		RHIBindDebugLabelName(TextureRHI, *Owner->GetName());
-		RHIUpdateTextureReference(Owner->TextureReference.TextureReferenceRHI, TextureRHI);
+		RHICmdList.BindDebugLabelName(TextureRHI, *Owner->GetName());
+		RHICmdList.UpdateTextureReference(Owner->TextureReference.TextureReferenceRHI, TextureRHI);
 	}
 	
 protected:
@@ -580,16 +579,18 @@ public:
 	 */
 	virtual void InitRHI(FRHICommandListBase& RHICmdList) override
 	{
+		const FString TextureName = Owner->GetName();
+
 		id<MTLTexture> MetalTexture = Owner->GetMetalTexture();
 		if (MetalTexture != nullptr)
 		{
 			Size.X = Size.Y = Owner->Size.X;
 
 			const FRHITextureCreateDesc Desc =
-				FRHITextureCreateDesc::CreateCube(TEXT("FARMetalResource"), Size.X, PF_B8G8R8A8)
+				FRHITextureCreateDesc::CreateCube(*TextureName, Size.X, PF_B8G8R8A8)
 				.SetFlags(ETextureCreateFlags::SRGB);
 
-			EnvCubemapTextureRHIRef = RHICreateTexture(Desc);
+			EnvCubemapTextureRHIRef = RHICmdList.CreateTexture(Desc);
 
 			/**
 			 * To map their texture faces into our space we need:
@@ -611,18 +612,15 @@ public:
 		{
 			Size.X = Size.Y = 1;
 
-			const FRHITextureCreateDesc Desc =
-				FRHITextureCreateDesc::CreateCube(TEXT("FARMetalResource"), Size.X, PF_B8G8R8A8);
+			const FRHITextureCreateDesc Desc = FRHITextureCreateDesc::CreateCube(*TextureName, Size.X, PF_B8G8R8A8);
 
 			// Start with a 1x1 texture
-			EnvCubemapTextureRHIRef = RHICreateTexture(Desc);
+			EnvCubemapTextureRHIRef = RHICmdList.CreateTexture(Desc);
 		}
 
 
 		TextureRHI = EnvCubemapTextureRHIRef;
-		TextureRHI->SetName(Owner->GetFName());
-		RHIBindDebugLabelName(TextureRHI, *Owner->GetName());
-		RHIUpdateTextureReference(Owner->TextureReference.TextureReferenceRHI, TextureRHI);
+		RHICmdList.UpdateTextureReference(Owner->TextureReference.TextureReferenceRHI, TextureRHI);
 
 		FSamplerStateInitializerRHI SamplerStateInitializer(SF_Bilinear, AM_Clamp, AM_Clamp, AM_Clamp);
 		SamplerStateRHI = RHICreateSamplerState(SamplerStateInitializer);
@@ -904,10 +902,10 @@ public:
 				.SetFlags(ETextureCreateFlags::ShaderResource);
 
 			FScopeLock ScopeLock(&DecodedTextureLock);
-			DecodedTextureRef = RHICreateTexture(Desc);
+			DecodedTextureRef = RHICmdList.CreateTexture(Desc);
 		}
 		
-		UpdateTextureRHI();
+		UpdateTextureRHI(RHICmdList);
 	}
 	
 	virtual void ReleaseRHI() override
@@ -929,12 +927,11 @@ public:
 		return Size.Y;
 	}
 	
-	void UpdateTextureRHI()
+	void UpdateTextureRHI(FRHICommandListBase& RHICmdList)
 	{
 		TextureRHI = DecodedTextureRef;
-		TextureRHI->SetName(Owner->GetFName());
-		RHIBindDebugLabelName(TextureRHI, *Owner->GetName());
-		RHIUpdateTextureReference(Owner->TextureReference.TextureReferenceRHI, TextureRHI);
+		RHICmdList.BindDebugLabelName(TextureRHI, *Owner->GetName());
+		RHICmdList.UpdateTextureReference(Owner->TextureReference.TextureReferenceRHI, TextureRHI);
 	}
 	
 #if SUPPORTS_ARKIT_1_0
@@ -964,8 +961,8 @@ public:
 				.SetInitActionBulkData(new FAppleARKitCameraTextureResourceWrapper(CapturedCbCrImage));
 
 			// pull the Y and CbCr textures out of the captured image planes (format is fake here, it will get the format from the FAppleARKitCameraTextureResourceWrapper)
-			VideoTextureY = RHICreateTexture(YDesc);
-			VideoTextureCbCr = RHICreateTexture(cBcRDesc);
+			VideoTextureY = RHICmdList.CreateTexture(YDesc);
+			VideoTextureCbCr = RHICmdList.CreateTexture(cBcRDesc);
 
 			// todo: Add an update call to the registry instead of this unregister/re-register
 			FExternalTextureRegistry::Get().UnregisterExternalTexture(ARKitPassthroughCameraExternalTextureYGuid);
@@ -1003,7 +1000,7 @@ public:
 					.SetFlags(ETextureCreateFlags::ShaderResource | ETextureCreateFlags::UAV);
 
 				FScopeLock ScopeLock(&DecodedTextureLock);
-				DecodedTextureRef = RHICreateTexture(Desc);
+				DecodedTextureRef = RHICmdList.CreateTexture(Desc);
 			}
 			
 			DecodedTextureUAV = nullptr;
@@ -1061,7 +1058,7 @@ public:
 			GraphBuilder.Execute();
 		}
 		
-		UpdateTextureRHI();
+		UpdateTextureRHI(RHICmdList);
 	}
 	
 	id<MTLTexture> GetMetalTexture()
diff --git a/Engine/Plugins/Runtime/AR/Google/GoogleARCore/Source/GoogleARCoreBase/Private/GoogleARCoreTexture.cpp b/Engine/Plugins/Runtime/AR/Google/GoogleARCore/Source/GoogleARCoreBase/Private/GoogleARCoreTexture.cpp
--- a/Engine/Plugins/Runtime/AR/Google/GoogleARCore/Source/GoogleARCoreBase/Private/GoogleARCoreTexture.cpp
+++ b/Engine/Plugins/Runtime/AR/Google/GoogleARCore/Source/GoogleARCoreBase/Private/GoogleARCoreTexture.cpp
@@ -143,17 +143,16 @@ public:
 		SamplerStateRHI = RHICreateSamplerState(SamplerStateInitializer);
 		
 		FString Name = Owner->GetName();
+
 		// Our source data is 16bits integer so PF_R32_FLOAT is a bit wasteful
 		// But PF_R16_UINT doesn't work in the material and it's not easy to convert the data to PF_R16F on the CPU
 		const FRHITextureCreateDesc Desc =
 			FRHITextureCreateDesc::Create2D(*Name, SizeX, SizeY, PF_R32_FLOAT)
 			.SetFlags(ETextureCreateFlags::ShaderResource);
 
-		TextureRHI = RHICreateTexture(Desc);
+		TextureRHI = RHICmdList.CreateTexture(Desc);
 
-		TextureRHI->SetName(Owner->GetFName());
-		RHIBindDebugLabelName(TextureRHI, *Name);
-		RHIUpdateTextureReference(Owner->TextureReference.TextureReferenceRHI, TextureRHI);
+		RHICmdList.UpdateTextureReference(Owner->TextureReference.TextureReferenceRHI, TextureRHI);
 	}
 	
 	virtual void ReleaseRHI() override
diff --git a/Engine/Plugins/TextureGraph/Source/TextureGraphEngine/Device/FX/Device_FX.cpp b/Engine/Plugins/TextureGraph/Source/TextureGraphEngine/Device/FX/Device_FX.cpp
--- a/Engine/Plugins/TextureGraph/Source/TextureGraphEngine/Device/FX/Device_FX.cpp
+++ b/Engine/Plugins/TextureGraph/Source/TextureGraphEngine/Device/FX/Device_FX.cpp
@@ -686,8 +686,10 @@ int32 Device_FX::InitRTResource(TexPtr TextureObj, UTextureRenderTarget* RT)
 	check(IsInRenderingThread());
 	check(RT->GetResource());
 
-	RT->GetResource()->InitResource(FRHICommandListImmediate::Get());
-	FDeferredUpdateResource::UpdateResources(RHI());
+	FRHICommandListImmediate& RHICmdList = FRHICommandListImmediate::Get();
+
+	RT->GetResource()->InitResource(RHICmdList);
+	FDeferredUpdateResource::UpdateResources(RHICmdList);
 
 	FTextureRenderTarget2DResource* RTRes = (FTextureRenderTarget2DResource*)RT->GetRenderTargetResource();
 	check(RTRes);
@@ -698,7 +700,7 @@ int32 Device_FX::InitRTResource(TexPtr TextureObj, UTextureRenderTarget* RT)
 	FTextureRHIRef RHITexture = RTRes->GetTextureRHI();
 	check(RHITexture);
 
-	RHIBindDebugLabelName(RHITexture, *TextureObj->GetDescriptor().Name);
+	RHICmdList.BindDebugLabelName(RHITexture, *TextureObj->GetDescriptor().Name);
 
 	return 0;
 }
diff --git a/Engine/Source/Runtime/Engine/Private/Particles/ParticleGpuSimulation.cpp b/Engine/Source/Runtime/Engine/Private/Particles/ParticleGpuSimulation.cpp
--- a/Engine/Source/Runtime/Engine/Private/Particles/ParticleGpuSimulation.cpp
+++ b/Engine/Source/Runtime/Engine/Private/Particles/ParticleGpuSimulation.cpp
@@ -345,7 +345,7 @@ public:
 		return PF_FloatRGBA;
 	}
 
-	FTextureRHIRef CreateStateTextureRHI(const TCHAR* Name, EPixelFormat PixelFormat)
+	FTextureRHIRef CreateStateTextureRHI(FRHICommandListBase& RHICmdList, const TCHAR* Name, EPixelFormat PixelFormat)
 	{
 		const FRHITextureCreateDesc Desc =
 			FRHITextureCreateDesc::Create2D(Name)
@@ -355,7 +355,7 @@ public:
 			.SetFlags(ETextureCreateFlags::RenderTargetable | ETextureCreateFlags::ShaderResource)
 			.SetInitialState(ERHIAccess::SRVMask);
 
-		return RHICreateTexture(Desc);
+		return RHICmdList.CreateTexture(Desc);
 	}
 
 	/**
@@ -365,11 +365,11 @@ public:
 	{
 		// 32-bit per channel RGBA texture for position.
 		check( !IsValidRef( PositionTextureRHI ) );
-		PositionTextureRHI = CreateStateTextureRHI(PaticlePositonTextureName, GetParticlePositionFormat());
+		PositionTextureRHI = CreateStateTextureRHI(RHICmdList, PaticlePositonTextureName, GetParticlePositionFormat());
 
 		// 16-bit per channel RGBA texture for velocity.
 		check(!IsValidRef(VelocityTextureRHI));
-		VelocityTextureRHI = CreateStateTextureRHI(PaticleVelocityTextureName, GetParticleVelocityFormat());
+		VelocityTextureRHI = CreateStateTextureRHI(RHICmdList, PaticleVelocityTextureName, GetParticleVelocityFormat());
 
 		// using FName's ability to append a number to a string (..._0) without an extra string allocation, except suffixing is done when number > 0 hence the +1 here : 
 		FName PositionTextureName(PaticlePositonTextureName, ParticleStateIndex + 1);
@@ -400,8 +400,8 @@ public:
 		RenderTargetSize.Y = SizeY;
 
 		// allocate the new textures
-		FTextureRHIRef NewPositionTextureRHI = CreateStateTextureRHI(PaticlePositonTextureName, PF_A32B32G32R32F);
-		FTextureRHIRef NewVelocityTextureRHI = CreateStateTextureRHI(PaticleVelocityTextureName, PF_FloatRGBA);
+		FTextureRHIRef NewPositionTextureRHI = CreateStateTextureRHI(RHICmdList, PaticlePositonTextureName, PF_A32B32G32R32F);
+		FTextureRHIRef NewVelocityTextureRHI = CreateStateTextureRHI(RHICmdList, PaticleVelocityTextureName, PF_FloatRGBA);
 		
 		RHICmdList.Transition(FRHITransitionInfo(NewPositionTextureRHI, ERHIAccess::SRVMask, ERHIAccess::CopyDest));
 		RHICmdList.Transition(FRHITransitionInfo(PositionTextureRHI, ERHIAccess::SRVMask, ERHIAccess::CopySrc));
@@ -444,7 +444,7 @@ public:
 
 	FIntPoint RenderTargetSize = FIntPoint::ZeroValue;
 
-	FTextureRHIRef CreateAttributesTextureRHI()
+	FTextureRHIRef CreateAttributesTextureRHI(FRHICommandListBase& RHICmdList)
 	{
 		const FRHITextureCreateDesc Desc =
 			FRHITextureCreateDesc::Create2D(PaticleAttributeTextureName)
@@ -454,23 +454,18 @@ public:
 			.SetFlags(ETextureCreateFlags::RenderTargetable | ETextureCreateFlags::ShaderResource | ETextureCreateFlags::NoFastClear)
 			.SetInitialState(ERHIAccess::RTV);
 
-		return RHICreateTexture(Desc);
+		return RHICmdList.CreateTexture(Desc);
 	}
 
 	/**
 	 * Initialize RHI resources used for particle simulation.
 	 */
-	virtual void InitRHI(FRHICommandListBase&) override
+	virtual void InitRHI(FRHICommandListBase& RHICmdList) override
 	{
-		TextureRHI = CreateAttributesTextureRHI();
-
- 		FName AttributesTextureName(PaticleAttributeTextureName);
- 		TextureRHI->SetName(AttributesTextureName);
- 		RHIBindDebugLabelName(TextureRHI, PaticleAttributeTextureName);
+		TextureRHI = CreateAttributesTextureRHI(RHICmdList);
 
-		FRHICommandListImmediate& RHICmdList = FRHICommandListExecutor::GetImmediateCommandList();
-		ClearRenderTarget(RHICmdList, TextureRHI);
-		RHICmdList.Transition(FRHITransitionInfo(TextureRHI, ERHIAccess::RTV, ERHIAccess::SRVMask));
+		ClearRenderTarget(RHICmdList.GetAsImmediate(), TextureRHI);
+		RHICmdList.GetAsImmediate().Transition(FRHITransitionInfo(TextureRHI, ERHIAccess::RTV, ERHIAccess::SRVMask));
 	}
 
 	/**
@@ -490,7 +485,7 @@ public:
 		RenderTargetSize.Y = SizeY;
 
 		// allocate the new textures
-		FTextureRHIRef NewAttributesTextureRHI = CreateAttributesTextureRHI();
+		FTextureRHIRef NewAttributesTextureRHI = CreateAttributesTextureRHI(RHICmdList);
 
 		RHICmdList.Transition(FRHITransitionInfo(NewAttributesTextureRHI, ERHIAccess::RTV, ERHIAccess::CopyDest));
 		RHICmdList.Transition(FRHITransitionInfo(TextureRHI, ERHIAccess::SRVMask, ERHIAccess::CopySrc));
diff --git a/Engine/Source/Runtime/Engine/Private/StaticMesh.cpp b/Engine/Source/Runtime/Engine/Private/StaticMesh.cpp
--- a/Engine/Source/Runtime/Engine/Private/StaticMesh.cpp
+++ b/Engine/Source/Runtime/Engine/Private/StaticMesh.cpp
@@ -1511,17 +1511,17 @@ void FStaticMeshLODResources::InitResources(UStaticMesh* Parent, int32 LODIndex)
 
 #if !(UE_BUILD_SHIPPING || UE_BUILD_TEST)
 	ENQUEUE_RENDER_COMMAND(NameRHIResources)(
-		[this, DebugName = (Parent ? Parent->GetFName() : NAME_None)](FRHICommandListImmediate&)
+		[this, DebugName = (Parent ? Parent->GetFName() : NAME_None)](FRHICommandListImmediate& RHICmdList)
 	{
 		TStringBuilder<512> StringBuilder;
-		auto SetDebugName = [&StringBuilder, DebugName](FRHIBuffer* RHIBuffer, const TCHAR* Extension)
+		auto SetDebugName = [&RHICmdList, &StringBuilder, DebugName](FRHIBuffer* RHIBuffer, const TCHAR* Extension)
 		{
 			if (RHIBuffer)
 			{
 				StringBuilder.Reset();
 				DebugName.ToString(StringBuilder);
 				StringBuilder.Append(Extension);
-				RHIBindDebugLabelName(RHIBuffer, StringBuilder.ToString());
+				RHICmdList.BindDebugLabelName(RHIBuffer, StringBuilder.ToString());
 			}
 		};
 		SetDebugName(IndexBuffer.IndexBufferRHI, TEXT("_IB"));
diff --git a/Engine/Source/Runtime/Engine/Private/Texture2D.cpp b/Engine/Source/Runtime/Engine/Private/Texture2D.cpp
--- a/Engine/Source/Runtime/Engine/Private/Texture2D.cpp
+++ b/Engine/Source/Runtime/Engine/Private/Texture2D.cpp
@@ -1771,11 +1771,11 @@ void FVirtualTexture2DResource::InitializeEditorResources(FRHICommandListBase& R
 			}
 
 			GraphBuilder.Execute();
+
+			RHICmdList.BindDebugLabelName(Texture2DRHI, *Name);
 		}
 
 		TextureRHI = Texture2DRHI;
-		TextureRHI->SetName(TextureName);
-		RHIBindDebugLabelName(TextureRHI, *Name);
 		
 		// Only update texture reference if necessary.
 		if (TextureRHI != TextureReferenceRHI->GetReferencedTexture())
diff --git a/Engine/Source/Runtime/Engine/Private/TextureCube.cpp b/Engine/Source/Runtime/Engine/Private/TextureCube.cpp
--- a/Engine/Source/Runtime/Engine/Private/TextureCube.cpp
+++ b/Engine/Source/Runtime/Engine/Private/TextureCube.cpp
@@ -471,9 +471,6 @@ public:
 		TextureCubeRHI = RHICreateTexture(Desc);
 
 		TextureRHI = TextureCubeRHI;
-		TextureRHI->SetOwnerName(GetOwnerName());
-		TextureRHI->SetName(Owner->GetFName());
-		RHIBindDebugLabelName(TextureRHI, *Name);
 		RHIUpdateTextureReference(Owner->TextureReference.TextureReferenceRHI,TextureRHI);
 
 		// Read the mip-levels into the RHI texture.
diff --git a/Engine/Source/Runtime/Engine/Private/TextureCubeArray.cpp b/Engine/Source/Runtime/Engine/Private/TextureCubeArray.cpp
--- a/Engine/Source/Runtime/Engine/Private/TextureCubeArray.cpp
+++ b/Engine/Source/Runtime/Engine/Private/TextureCubeArray.cpp
@@ -147,10 +147,7 @@ public:
 
 		TextureCubeRHI = RHICreateTexture(Desc);
 		TextureRHI = TextureCubeRHI;
-		TextureRHI->SetOwnerName(GetOwnerName());
-		TextureRHI->SetName(Owner->GetFName());
 
-		RHIBindDebugLabelName(TextureRHI, *Name);
 		RHIUpdateTextureReference(Owner->TextureReference.TextureReferenceRHI, TextureRHI);
 
 		// Read the mip-levels into the RHI texture.
diff --git a/Engine/Source/Runtime/RHI/Private/RHIValidation.cpp b/Engine/Source/Runtime/RHI/Private/RHIValidation.cpp
--- a/Engine/Source/Runtime/RHI/Private/RHIValidation.cpp
+++ b/Engine/Source/Runtime/RHI/Private/RHIValidation.cpp
@@ -1186,10 +1186,9 @@ static FString GetBreadcrumbPath()
 // FlushType: Thread safe
 void FValidationRHI::RHIBindDebugLabelName(FRHICommandListBase& RHICmdList, FRHITexture* Texture, const TCHAR* Name)
 {
-	FString NameCopyRT = Name;
-	RHICmdList.EnqueueLambda([Texture, NameCopyRHIT = MoveTemp(NameCopyRT)](FRHICommandListBase& RHICmdList)
+	RHICmdList.EnqueueLambda([Texture, Name = FString(Name)](FRHICommandListBase& RHICmdList)
 	{
-		((FValidationContext&)RHICmdList.GetComputeContext()).Tracker->Rename(Texture->GetTrackerResource(), *NameCopyRHIT);
+		((FValidationContext&)RHICmdList.GetComputeContext()).Tracker->Rename(Texture->GetTrackerResource(), *Name);
 	});
 
 	RHI->RHIBindDebugLabelName(RHICmdList, Texture, Name);
@@ -1197,10 +1196,9 @@ void FValidationRHI::RHIBindDebugLabelName(FRHICommandListBase& RHICmdList, FRHI
 
 void FValidationRHI::RHIBindDebugLabelName(FRHICommandListBase& RHICmdList, FRHIBuffer* Buffer, const TCHAR* Name)
 {
-	FString NameCopyRT = Name;
-	RHICmdList.EnqueueLambda([Buffer, NameCopyRHIT = MoveTemp(NameCopyRT)](FRHICommandListBase& RHICmdList)
+	RHICmdList.EnqueueLambda([Buffer, Name = FString(Name)](FRHICommandListBase& RHICmdList)
 	{
-		((FValidationContext&)RHICmdList.GetComputeContext()).Tracker->Rename(Buffer, *NameCopyRHIT);
+		((FValidationContext&)RHICmdList.GetComputeContext()).Tracker->Rename(Buffer, *Name);
 	});
 
 	RHI->RHIBindDebugLabelName(RHICmdList, Buffer, Name);
@@ -1209,10 +1207,9 @@ void FValidationRHI::RHIBindDebugLabelName(FRHICommandListBase& RHICmdList, FRHI
 void FValidationRHI::RHIBindDebugLabelName(FRHICommandListBase& RHICmdList, FRHIUnorderedAccessView* UnorderedAccessViewRHI, const TCHAR* Name)
 {
 	RHIValidation::FResource* Resource = UnorderedAccessViewRHI->GetViewIdentity().Resource;
-	FString NameCopyRT = Name;
-	RHICmdList.EnqueueLambda([Resource, NameCopyRHIT = MoveTemp(NameCopyRT)](FRHICommandListBase& RHICmdList)
+	RHICmdList.EnqueueLambda([Resource, Name = FString(Name)](FRHICommandListBase& RHICmdList)
 	{
-		((FValidationContext&)RHICmdList.GetComputeContext()).Tracker->Rename(Resource, *NameCopyRHIT);
+		((FValidationContext&)RHICmdList.GetComputeContext()).Tracker->Rename(Resource, *Name);
 	});
 
 	RHI->RHIBindDebugLabelName(RHICmdList, UnorderedAccessViewRHI, Name);
diff --git a/Engine/Source/Runtime/RHI/Public/RHICommandList.h b/Engine/Source/Runtime/RHI/Public/RHICommandList.h
--- a/Engine/Source/Runtime/RHI/Public/RHICommandList.h
+++ b/Engine/Source/Runtime/RHI/Public/RHICommandList.h
@@ -5706,16 +5706,19 @@ FORCEINLINE FRayTracingAccelerationStructureOfflineMetadata RHIGetRayTracingGeom
 	return GDynamicRHI->RHIGetRayTracingGeometryOfflineMetadata(OfflineDataHeader);
 }
 
+UE_DEPRECATED(5.7, "BindDebugLabelName with an implied immediate command list is deprecated")
 FORCEINLINE void RHIBindDebugLabelName(FRHITexture* Texture, const TCHAR* Name)
 {
 	FRHICommandListImmediate::Get().BindDebugLabelName(Texture, Name);
 }
 
+UE_DEPRECATED(5.7, "BindDebugLabelName with an implied immediate command list is deprecated")
 FORCEINLINE void RHIBindDebugLabelName(FRHIBuffer* Buffer, const TCHAR* Name)
 {
 	FRHICommandListImmediate::Get().BindDebugLabelName(Buffer, Name);
 }
 
+UE_DEPRECATED(5.7, "BindDebugLabelName with an implied immediate command list is deprecated")
 FORCEINLINE void RHIBindDebugLabelName(FRHIUnorderedAccessView* UnorderedAccessViewRHI, const TCHAR* Name)
 {
 	FRHICommandListImmediate::Get().BindDebugLabelName(UnorderedAccessViewRHI, Name);
