Metal RHI

Fix AppleARKit not compiling after adding IMetalDynamicRHI.h

diff --git a/Engine/Plugins/Runtime/AR/AppleAR/AppleARKit/Source/AppleARKit/AppleARKit.Build.cs b/Engine/Plugins/Runtime/AR/AppleAR/AppleARKit/Source/AppleARKit/AppleARKit.Build.cs
--- a/Engine/Plugins/Runtime/AR/AppleAR/AppleARKit/Source/AppleARKit/AppleARKit.Build.cs
+++ b/Engine/Plugins/Runtime/AR/AppleAR/AppleARKit/Source/AppleARKit/AppleARKit.Build.cs
@@ -18,6 +18,7 @@ public class AppleARKit : ModuleRules
 			"Engine",
 			"MRMesh",
 			"EyeTracker",
+			"MetalRHI",
 		});
 			
 		
@@ -36,8 +37,10 @@ public class AppleARKit : ModuleRules
 			"Projects",
 			"ARUtilities",
 		});
-		
-		
+
+		PublicWeakFrameworks.Add("Metal");
+		AddEngineThirdPartyPrivateStaticDependencies(Target, "MetalCPP");
+            
 		DynamicallyLoadedModuleNames.AddRange(new string[]
 		{
 		});
diff --git a/Engine/Plugins/Runtime/AR/AppleAR/AppleARKit/Source/AppleARKit/Private/AppleARKitTextures.cpp b/Engine/Plugins/Runtime/AR/AppleAR/AppleARKit/Source/AppleARKit/Private/AppleARKitTextures.cpp
--- a/Engine/Plugins/Runtime/AR/AppleAR/AppleARKit/Source/AppleARKit/Private/AppleARKitTextures.cpp
+++ b/Engine/Plugins/Runtime/AR/AppleAR/AppleARKit/Source/AppleARKit/Private/AppleARKitTextures.cpp
@@ -15,6 +15,12 @@
 #include "TextureResource.h"
 #include "Materials/MaterialRenderProxy.h"
 
+THIRD_PARTY_INCLUDES_START
+#include "MetalInclude.h"
+THIRD_PARTY_INCLUDES_END
+
+#include "IMetalDynamicRHI.h"
+
 #if PLATFORM_APPLE
 	#import <Metal/Metal.h>
 	#import <MetalPerformanceShaders/MetalPerformanceShaders.h>
@@ -839,11 +845,11 @@ public:
 		{
 			IMetalDynamicRHI* MetalDynamicRHI = GetIMetalDynamicRHI();
 			
-			VideoTextureY = MetalDynamicRHI->RHICreateTexture2DFromCVMetalTexture(PF_B8G8R8A8, ETextureCreateFlags::NoTiling | ETextureCreateFlags::ShaderResource, FClearValueBinding::White, CapturedYTexture);
-			CFRelease(CapturedYTexture);
+			VideoTextureY = MetalDynamicRHI->RHICreateTexture2DFromCVMetalTexture(PF_B8G8R8A8, ETextureCreateFlags::NoTiling | ETextureCreateFlags::ShaderResource, FClearValueBinding::White, CapturedYImage);
+			CFRelease(CapturedYImage);
 	
-			VideoTextureCbCr = MetalDynamicRHI->RHICreateTexture2DFromCVMetalTexture(PF_B8G8R8A8, ETextureCreateFlags::NoTiling | ETextureCreateFlags::ShaderResource, FClearValueBinding::White, CapturedCbCrTexture);
-			CFRelease(CapturedCbCrTexture);
+			VideoTextureCbCr = MetalDynamicRHI->RHICreateTexture2DFromCVMetalTexture(PF_B8G8R8A8, ETextureCreateFlags::NoTiling | ETextureCreateFlags::ShaderResource, FClearValueBinding::White, CapturedCbCrImage);
+			CFRelease(CapturedCbCrImage);
 
 			// todo: Add an update call to the registry instead of this unregister/re-register
 			FExternalTextureRegistry::Get().UnregisterExternalTexture(ARKitPassthroughCameraExternalTextureYGuid);
