Removed MetalRHI from AppleARKit plugin

diff --git a/Engine/Plugins/Runtime/AR/AppleAR/AppleARKit/Source/AppleARKit/AppleARKit.Build.cs b/Engine/Plugins/Runtime/AR/AppleAR/AppleARKit/Source/AppleARKit/AppleARKit.Build.cs
--- a/Engine/Plugins/Runtime/AR/AppleAR/AppleARKit/Source/AppleARKit/AppleARKit.Build.cs
+++ b/Engine/Plugins/Runtime/AR/AppleAR/AppleARKit/Source/AppleARKit/AppleARKit.Build.cs
@@ -17,10 +17,8 @@ public class AppleARKit : ModuleRules
 			"Core",
 			"Engine",
 			"MRMesh",
-			"EyeTracker",
-			"MetalRHI",
+			"EyeTracker",			
 		});
-			
 		
 		PrivateDependencyModuleNames.AddRange(new string[]
 		{
@@ -37,9 +35,6 @@ public class AppleARKit : ModuleRules
 			"Projects",
 			"ARUtilities",
 		});
-
-		PublicWeakFrameworks.Add("Metal");
-		AddEngineThirdPartyPrivateStaticDependencies(Target, "MetalCPP");
             
 		DynamicallyLoadedModuleNames.AddRange(new string[]
 		{
@@ -55,14 +50,25 @@ public class AppleARKit : ModuleRules
 				"ARKit",
 				"MetalPerformanceShaders",
 				"CoreLocation",
+				"MetalRHI",
 			});
 			
+			PublicWeakFrameworks.Add("Metal");
+			AddEngineThirdPartyPrivateStaticDependencies(Target, "MetalCPP");
+		
 			string PluginPath = Utils.MakePathRelativeTo(ModuleDirectory, Target.RelativeEnginePath);
             AdditionalPropertiesForReceipt.Add("IOSPlugin", Path.Combine(PluginPath, "AppleARKit_IOS_UPL.xml"));
 		}
 		else if (Target.Platform == UnrealTargetPlatform.Mac)
 		{
-			PublicFrameworks.Add("MetalPerformanceShaders");
+			PublicFrameworks.AddRange(new string[]
+			{
+				"MetalPerformanceShaders",
+				"MetalRHI",
+			});
+			
+			PublicWeakFrameworks.Add("Metal");
+			AddEngineThirdPartyPrivateStaticDependencies(Target, "MetalCPP");
 		}
 	}
 }
