Fix MetalRHI not being included on Mac platform for ARKit

diff --git a/Engine/Plugins/Runtime/AR/AppleAR/AppleARKit/Source/AppleARKit/AppleARKit.Build.cs b/Engine/Plugins/Runtime/AR/AppleAR/AppleARKit/Source/AppleARKit/AppleARKit.Build.cs
--- a/Engine/Plugins/Runtime/AR/AppleAR/AppleARKit/Source/AppleARKit/AppleARKit.Build.cs
+++ b/Engine/Plugins/Runtime/AR/AppleAR/AppleARKit/Source/AppleARKit/AppleARKit.Build.cs
@@ -50,9 +50,12 @@ public class AppleARKit : ModuleRules
 				"ARKit",
 				"MetalPerformanceShaders",
 				"CoreLocation",
-				"MetalRHI",
 			});
 			
+			PrivateDependencyModuleNames.AddRange(new string[] {
+        	    "MetalRHI",
+        	});
+        	
 			PublicWeakFrameworks.Add("Metal");
 			AddEngineThirdPartyPrivateStaticDependencies(Target, "MetalCPP");
 		
@@ -64,9 +67,12 @@ public class AppleARKit : ModuleRules
 			PublicFrameworks.AddRange(new string[]
 			{
 				"MetalPerformanceShaders",
-				"MetalRHI",
 			});
 			
+			PrivateDependencyModuleNames.AddRange(new string[] {
+        	    "MetalRHI",
+        	});
+			
 			PublicWeakFrameworks.Add("Metal");
 			AddEngineThirdPartyPrivateStaticDependencies(Target, "MetalCPP");
 		}
