Fix CachedUniformExpressionShaderMap not being populated for AR template with Scene Depth occlusion enabled. Drive by removal of undefined static variables.

diff --git a/Engine/Plugins/Runtime/AR/AppleAR/AppleARKit/Source/AppleARKit/Private/AppleARKitVideoOverlay.cpp b/Engine/Plugins/Runtime/AR/AppleAR/AppleARKit/Source/AppleARKit/Private/AppleARKitVideoOverlay.cpp
--- a/Engine/Plugins/Runtime/AR/AppleAR/AppleARKit/Source/AppleARKit/Private/AppleARKitVideoOverlay.cpp
+++ b/Engine/Plugins/Runtime/AR/AppleAR/AppleARKit/Source/AppleARKit/Private/AppleARKitVideoOverlay.cpp
@@ -350,6 +350,12 @@ void FAppleARKitVideoOverlay::RenderVideoOverlayWithMaterial(FRHICommandList& RH
 	const FMaterial& CameraMaterial = MaterialProxy->GetMaterialWithFallback(FeatureLevel, MaterialProxy);
 	const FMaterialShaderMap* const MaterialShaderMap = CameraMaterial.GetRenderingThreadShaderMap();
 
+	// If there are pending MI parameter updates, ensure these get processed
+	if (MaterialProxy->UniformExpressionCache[FeatureLevel].CachedUniformExpressionShaderMap == nullptr)
+	{
+		FMaterialRenderProxy::UpdateDeferredCachedUniformExpressions(RHICmdList);
+	}
+	
 	FGraphicsPipelineStateInitializer GraphicsPSOInit;
 	RHICmdList.ApplyCachedRenderTargets(GraphicsPSOInit);
 
@@ -391,7 +397,7 @@ void FAppleARKitVideoOverlay::RenderVideoOverlayWithMaterial(FRHICommandList& RH
 	}
 
 	SetGraphicsPipelineState(RHICmdList, GraphicsPSOInit, 0);
-
+	
 	if (bIsMobileRenderer)
 	{
 		TShaderRef<FARKitCameraOverlayMobileVS> VertexShaderPtr(static_cast<FARKitCameraOverlayMobileVS*>(VertexShader.GetShader()), *MaterialShaderMap);
diff --git a/Engine/Source/Runtime/Renderer/Public/MaterialShader.h b/Engine/Source/Runtime/Renderer/Public/MaterialShader.h
--- a/Engine/Source/Runtime/Renderer/Public/MaterialShader.h
+++ b/Engine/Source/Runtime/Renderer/Public/MaterialShader.h
@@ -113,11 +113,6 @@ public:
 		FMeshDrawSingleShaderBindings& ShaderBindings) const;
 
 private:
-	/** If true, cached uniform expressions are allowed. */
-	static RENDERER_API int32 bAllowCachedUniformExpressions;
-	/** Console variable ref to toggle cached uniform expressions. */
-	static RENDERER_API FAutoConsoleVariableRef CVarAllowCachedUniformExpressions;
-
 #if !(UE_BUILD_TEST || UE_BUILD_SHIPPING || !WITH_EDITOR)
 	RENDERER_API void VerifyExpressionAndShaderMaps(const FMaterialRenderProxy* MaterialRenderProxy, const FMaterial& Material, const FUniformExpressionCache* UniformExpressionCache) const;
 #endif
