Fix non-unity build error

diff --git a/Engine/Plugins/Runtime/AR/AppleAR/AppleARKit/Source/AppleARKit/Private/AppleARKitVideoOverlay.cpp b/Engine/Plugins/Runtime/AR/AppleAR/AppleARKit/Source/AppleARKit/Private/AppleARKitVideoOverlay.cpp
--- a/Engine/Plugins/Runtime/AR/AppleAR/AppleARKit/Source/AppleARKit/Private/AppleARKitVideoOverlay.cpp
+++ b/Engine/Plugins/Runtime/AR/AppleAR/AppleARKit/Source/AppleARKit/Private/AppleARKitVideoOverlay.cpp
@@ -35,6 +35,10 @@
 	#include "IOSAppDelegate.h"
 #endif
 
+#if SUPPORTS_ARKIT_3_0
+	#include "IMetalDynamicRHI.h"
+#endif
+
 DECLARE_FLOAT_COUNTER_STAT(TEXT("ARKit Frame to Texture Delay (ms)"), STAT_ARKitFrameToTextureDelay, STATGROUP_ARKIT);
 DECLARE_FLOAT_COUNTER_STAT(TEXT("ARKit Frame to Render Delay (ms)"), STAT_ARKitFrameToRenderDelay, STATGROUP_ARKIT);
 DECLARE_CYCLE_STAT(TEXT("Update Occlusion Textures"), STAT_UpdateOcclusionTextures, STATGROUP_ARKIT);
